[Feature] Migrate ReviewChangelistTool to UE5

diff --git a/Engine/Plugins/Editor/ChangelistReview/ChangelistReview.uplugin b/Engine/Plugins/Editor/ChangelistReview/ChangelistReview.uplugin
new file mode 100644
--- /dev/null
+++ b/Engine/Plugins/Editor/ChangelistReview/ChangelistReview.uplugin
@@ -0,0 +1,27 @@
+{
+	"FileVersion": 1,
+	"Version": 1,
+	"VersionName": "1.0",
+	"FriendlyName": "Changelist Reviews",
+	"Description": "Review source control changelists",
+	"Category": "Source Control",
+	"CreatedBy": "Epic Games, Inc.",
+	"CreatedByURL": "http://epicgames.com",
+	"DocsURL": "",
+	"EnabledByDefault": true,
+	"MarketplaceURL": "",
+	"SupportURL": "",
+	"CanContainContent": false,
+	"IsBetaVersion": false,
+	"IsExperimentalVersion": false,
+	"Installed": false,
+	"Plugins": [
+	],
+	"Modules": [
+		{
+			"Name": "ChangelistReview",
+			"Type": "Editor",
+			"LoadingPhase": "Default"
+		}
+	]
+}
diff --git a/Engine/Plugins/Editor/ChangelistReview/Source/ChangelistReview.Build.cs b/Engine/Plugins/Editor/ChangelistReview/Source/ChangelistReview.Build.cs
new file mode 100644
--- /dev/null
+++ b/Engine/Plugins/Editor/ChangelistReview/Source/ChangelistReview.Build.cs
@@ -0,0 +1,30 @@
+// Copyright Epic Games, Inc. All Rights Reserved.
+
+using UnrealBuildTool;
+
+public class ChangelistReview : ModuleRules
+{
+	public ChangelistReview(ReadOnlyTargetRules Target) : base(Target)
+	{
+		PublicDependencyModuleNames.AddRange(
+			new string[]
+			{
+				"Core",
+				"Engine",
+				"CoreUObject",
+			}
+		);
+
+		PrivateDependencyModuleNames.AddRange(
+			new string[] {
+				"SlateCore",
+				"Slate",
+				"EditorFramework",
+				"UnrealEd",
+				"EditorStyle",
+				"ToolMenus",
+				"SourceControl",
+			}
+		);
+	}
+}
\ No newline at end of file
diff --git a/Engine/Plugins/Editor/ChangelistReview/Source/Private/ChangelistReviewModule.cpp b/Engine/Plugins/Editor/ChangelistReview/Source/Private/ChangelistReviewModule.cpp
new file mode 100644
--- /dev/null
+++ b/Engine/Plugins/Editor/ChangelistReview/Source/Private/ChangelistReviewModule.cpp
@@ -0,0 +1,84 @@
+// Copyright Epic Games, Inc. All Rights Reserved.
+
+#include "ChangelistReviewModule.h"
+#include "ISourceControlModule.h"
+#include "Widgets/Docking/SDockTab.h"
+#include "ISourceControlProvider.h"
+#include "SSourceControlReview.h"
+#include "ToolMenus.h"
+
+#define LOCTEXT_NAMESPACE "FChangelistReviewModule"
+
+static const FName SourceControlReviewTabName = FName(TEXT("SourceControlChangelistReview"));
+
+void FChangelistReviewModule::StartupModule()
+{
+	UToolMenu* SourceControlMenu = UToolMenus::Get()->ExtendMenu("StatusBar.ToolBar.SourceControl");
+	FToolMenuSection& Section = SourceControlMenu->AddSection("SourceControlActions", LOCTEXT("SourceControlMenuHeadingActions", "Actions"));
+	
+	const FSlateIcon SourceControlIcon(FAppStyle::GetAppStyleSetName(), "SourceControl.ChangelistsTab");
+	Section.AddMenuEntry(
+		FName(TEXT("Review Changelists")),
+		LOCTEXT("ReviewChangelist_Label", "Review Changelists"),
+		LOCTEXT("ReviewChangelist_Tooltip", "Opens a dialog to review shelved and submitted changelists"),
+		SourceControlIcon,
+		FUIAction(
+			FExecuteAction::CreateRaw(this, &FChangelistReviewModule::ShowReviewTab),
+			FCanExecuteAction::CreateRaw(this, &FChangelistReviewModule::CanShowReviewTab)
+		)
+	);
+	
+	// Register the changelist tab spawner
+	FGlobalTabmanager::Get()->RegisterNomadTabSpawner(SourceControlReviewTabName, FOnSpawnTab::CreateRaw(this, &FChangelistReviewModule::CreateReviewTab))
+		.SetDisplayName(LOCTEXT("ChangelistsTabTitle", "Review Changelist"))
+		.SetTooltipText(LOCTEXT("ChangelistsTabTooltip", "Opens a dialog to review shelved and submitted changelists."))
+		.SetIcon(SourceControlIcon)
+		.SetMenuType(ETabSpawnerMenuType::Hidden);
+}
+
+void FChangelistReviewModule::ShutdownModule()
+{
+}
+
+TSharedRef<SDockTab> FChangelistReviewModule::CreateReviewTab(const FSpawnTabArgs& Args)
+{
+	return SAssignNew(ReviewTab, SDockTab)
+		.TabRole(ETabRole::NomadTab)
+		[
+			CreateReviewUI().ToSharedRef()
+		];
+}
+
+TSharedPtr<SWidget> FChangelistReviewModule::CreateReviewUI()
+{
+	TSharedPtr<SWidget> ReturnWidget;
+	if (IsInGameThread())
+	{
+		TSharedPtr<SSourceControlReview> SharedPtr = SNew(SSourceControlReview);
+		ReturnWidget = SharedPtr;
+	}
+	 return ReturnWidget;
+}
+
+void FChangelistReviewModule::ShowReviewTab()
+{
+	FGlobalTabmanager::Get()->TryInvokeTab(FTabId(SourceControlReviewTabName));
+}
+
+bool FChangelistReviewModule::CanShowReviewTab() const
+{
+	const ISourceControlModule& SourceControlModule = ISourceControlModule::Get();
+	if (SourceControlModule.IsEnabled())
+	{
+		const ISourceControlProvider& Provider = SourceControlModule.GetProvider();
+	
+		// we currently only support perforce for the review tool
+		return Provider.IsAvailable() && Provider.GetName().ToString() == TEXT("Perforce");
+	}
+	return false;
+}
+
+
+#undef LOCTEXT_NAMESPACE
+	
+IMPLEMENT_MODULE(FChangelistReviewModule, ChangelistReview)
\ No newline at end of file
diff --git a/Engine/Plugins/Editor/ChangelistReview/Source/Private/SChangelistEditableText.cpp b/Engine/Plugins/Editor/ChangelistReview/Source/Private/SChangelistEditableText.cpp
new file mode 100644
--- /dev/null
+++ b/Engine/Plugins/Editor/ChangelistReview/Source/Private/SChangelistEditableText.cpp
@@ -0,0 +1,76 @@
+// Copyright Epic Games, Inc. All Rights Reserved.
+
+#include "SChangelistEditableText.h"
+
+#include "Containers/UnrealString.h"
+#include "Framework/SlateDelegates.h"
+#include "HAL/Platform.h"
+#include "Internationalization/Text.h"
+#include "Templates/UniquePtr.h"
+#include "Widgets/Text/SlateEditableTextLayout.h"
+
+TDefaultNumericTypeInterface<int64> SChangelistEditableText::NumericTypeUtil;
+
+void SChangelistEditableText::Construct(const FArguments& InArgs)
+{
+	SEditableText::Construct(InArgs);
+}
+
+void SChangelistEditableText::OnTextCommitted(const FText& InText, const ETextCommit::Type InTextAction)
+{
+	TOptional<int64> NewValue = NumericTypeUtil.FromString(InText.ToString(), ValueAttribute.Get());
+	
+	if (NewValue.IsSet())
+	{
+		if (!ValueAttribute.IsBound())
+		{
+			ValueAttribute.Set(NewValue.GetValue());
+		}
+
+		OnTextCommittedCallback.ExecuteIfBound(InText, InTextAction);
+	}
+}
+
+void SChangelistEditableText::OnTextChanged(const FText& InText)
+{
+	const FString& Data = InText.ToString();
+
+	// Use the longest substring that consists of only valid characters.
+	// For example if someone enters:
+	// "john.doe2/CL_123456789/version_13", we'll use "123456789" because it's longer than "2" and "13"
+	int32 BestSubstrBegin = 0;
+	int32 BestSubstrLen = 0;
+	int32 I = 0;
+	while (I < Data.Len())
+	{
+		// skip invalid characters
+		while(I < Data.Len() && !IsCharacterValid(Data[I]))
+		{
+			++I;
+		}
+
+		// grab chunk of consecutive valid characters
+		const int32 SubstrBegin = I;
+		while(I < Data.Len() && IsCharacterValid(Data[I]))
+		{
+			++I;
+		}
+		
+		// if found substr is longer than the best so far, update the best so far
+		if (I - SubstrBegin > BestSubstrLen)
+		{
+			BestSubstrLen = I - SubstrBegin;
+			BestSubstrBegin = SubstrBegin;
+		}
+	}
+	
+	const FString ValidData = Data.Mid(BestSubstrBegin, BestSubstrLen);
+	const FText ValidText = FText::FromString(ValidData);
+	EditableTextLayout->SetText(ValidText);
+	OnTextChangedCallback.ExecuteIfBound(ValidText);
+}
+
+bool SChangelistEditableText::IsCharacterValid(TCHAR InChar)
+{
+	return TChar<TCHAR>::IsDigit(InChar);
+}
diff --git a/Engine/Plugins/Editor/ChangelistReview/Source/Private/SChangelistEditableText.h b/Engine/Plugins/Editor/ChangelistReview/Source/Private/SChangelistEditableText.h
new file mode 100644
--- /dev/null
+++ b/Engine/Plugins/Editor/ChangelistReview/Source/Private/SChangelistEditableText.h
@@ -0,0 +1,25 @@
+// Copyright Epic Games, Inc. All Rights Reserved.
+
+#pragma once
+#include "Widgets/Input/SEditableText.h"
+#include "Widgets/Input/NumericTypeInterface.h"
+
+class FText;
+
+/**
+ * Blueprint review changelist slate widget 
+ */
+class SChangelistEditableText : public SEditableText
+{
+public:
+	void Construct(const FArguments& InArgs);
+private:
+	//~ Begin ISlateEditableTextWidget Interface
+	virtual void OnTextChanged(const FText& InText) override;
+	virtual void OnTextCommitted(const FText& InText, const ETextCommit::Type InTextAction) override;
+	//~ End ISlateEditableTextWidget Interface
+	
+	TAttribute<int64> ValueAttribute;
+	static TDefaultNumericTypeInterface<int64> NumericTypeUtil;
+	static bool IsCharacterValid(TCHAR InChar);
+};
diff --git a/Engine/Plugins/Editor/ChangelistReview/Source/Private/SSourceControlReview.cpp b/Engine/Plugins/Editor/ChangelistReview/Source/Private/SSourceControlReview.cpp
new file mode 100644
--- /dev/null
+++ b/Engine/Plugins/Editor/ChangelistReview/Source/Private/SSourceControlReview.cpp
@@ -0,0 +1,560 @@
+// Copyright Epic Games, Inc. All Rights Reserved.
+
+#include "SSourceControlReview.h"
+
+#include "ISourceControlModule.h"
+#include "ISourceControlProvider.h"
+#include "SChangelistEditableText.h"
+#include "SourceControlOperations.h"
+#include "SSourceControlReviewEntry.h"
+#include "AssetRegistry/AssetRegistryModule.h"
+#include "Kismet/KismetSystemLibrary.h"
+#include "Styling/StarshipCoreStyle.h"
+#include "Widgets/Input/SMultiLineEditableTextBox.h"
+#include "Widgets/Layout/SScrollBox.h"
+#include "Widgets/Layout/SWidgetSwitcher.h"
+#include "Widgets/Notifications/SProgressBar.h"
+#include "Widgets/Input/SButton.h"
+#include "Widgets/Layout/SBox.h"
+#include "Misc/App.h"
+#include "Misc/MessageDialog.h"
+
+#define LOCTEXT_NAMESPACE "SourceControlReview"
+
+namespace
+{
+	const FText EnterChangelistText(LOCTEXT("EnterChangelistText", "Enter a changelist number to diff:"));
+	const FText EnterChangelistTooltip(LOCTEXT("EnterChangelistTooltip", "Enter changelist"));
+	const FText LoadChangelistText(LOCTEXT("LoadChangelistText", "Load Changelist"));
+	const FText LoadingText(LOCTEXT("LoadingText", "Loading..."));
+	
+	const FString FileDepotKey = TEXT("depotFile");
+	const FString FileRevisionKey = TEXT("rev");
+	const FString FileActionKey = TEXT("action");
+	const FString TimeKey = TEXT("time");
+	const FString AuthorKey = TEXT("user");
+	const FString DescriptionKey = TEXT("desc");
+	const FString ChangelistStatusKey = TEXT("status");
+	const FString ChangelistPendingStatusKey = TEXT("pending");
+	constexpr int32 RecordIndex = 0;
+}
+
+void SSourceControlReview::Construct(const FArguments& InArgs)
+{
+	const FString ProjectName = FApp::GetProjectName();
+	ensureAlwaysMsgf((!ProjectName.IsEmpty()), TEXT("BlueprintReviewTool - Unable to get ProjectName"));
+	
+	ChildSlot
+	[
+		SNew(SBorder)
+		.BorderImage(FAppStyle::GetBrush("Brushes.Panel"))
+		.Padding(10.f)
+		[
+			SNew(SVerticalBox)
+			+SVerticalBox::Slot()
+			.AutoHeight()
+			[
+				SNew(SHorizontalBox)
+				+SHorizontalBox::Slot()
+				.VAlign(VAlign_Center)
+				.AutoWidth()
+				[
+					SNew(STextBlock)
+					.Text(EnterChangelistText)
+					.Font(FStyleFonts::Get().LargeBold)
+				]
+				+SHorizontalBox::Slot()
+				.VAlign(VAlign_Center)
+				.Padding(5.f,0.f)
+				.AutoWidth()
+				[
+					SNew(SBorder)
+					.ToolTipText(EnterChangelistTooltip)
+					[
+						SAssignNew(ChangelistNumWidget, SChangelistEditableText)
+						.Font(FStyleFonts::Get().Large)
+						.MinDesiredWidth(200.f)
+						.Justification(ETextJustify::Center)
+						.Style(&FAppStyle::Get().GetWidgetStyle<FEditableTextStyle>("NormalEditableText"))
+						.OnTextCommitted(this, &SSourceControlReview::OnChangelistNumCommitted)
+					]
+				]
+				+SHorizontalBox::Slot()
+				.VAlign(VAlign_Center)
+				[
+					SNew(SButton)
+					.HAlign(HAlign_Center)
+					.OnClicked(this, &SSourceControlReview::OnLoadChangelistClicked)
+					[
+						SNew(STextBlock)
+						.Text(LoadChangelistText)
+						.Font(FStyleFonts::Get().LargeBold)
+					]
+				]
+			]
+			+SVerticalBox::Slot()
+			.Padding(0.f, 10.f)
+			.HAlign(HAlign_Center)
+			.VAlign(VAlign_Center)
+			.AutoHeight()
+			[
+				SAssignNew(LoadingTextBlock, STextBlock)
+				.Visibility(EVisibility::Collapsed)
+				.Text(LoadingText)
+				.Font(FStyleFonts::Get().Large)
+			]
+			+SVerticalBox::Slot()
+			.Padding(0.f, 15.f)
+			.AutoHeight()
+			[
+				SAssignNew(LoadingProgressBar, SProgressBar)
+				.Visibility(EVisibility::Collapsed)
+				.Percent(1.f)
+			]
+			+SVerticalBox::Slot()
+			.AutoHeight()
+			.Padding(0.f, 20.f)
+			[
+				GetChangelistInfoWidget()
+			]
+			+SVerticalBox::Slot()
+			[
+				SNew(SScrollBox)
+				+SScrollBox::Slot()
+				[
+					SAssignNew(ChangelistEntriesWidget, SVerticalBox)
+				]
+			]
+		]
+	];
+}
+
+TSharedRef<SWidget> SSourceControlReview::GetChangelistInfoWidget()
+{
+	auto GenerateRow = [](FText Key, const TAttribute<FText>& Value)->TSharedRef<SWidget>
+	{
+		const TSharedRef<SWidget> KeyWidget = SNew(STextBlock)
+			.Text(FText::Format(LOCTEXT("ChangelistInfoKey", "{0}:"), Key))
+			.Font(FStyleFonts::Get().Large);
+
+		const TSharedRef<SWidget> ValueWidget = SNew(SBox)
+			.MaxDesiredHeight(125.f)
+			[
+				SNew(SScrollBox)
+				+SScrollBox::Slot()
+				[
+					SNew(SMultiLineEditableTextBox)
+					.Text(Value)
+					.IsReadOnly(true)
+					.Font(FStyleFonts::Get().Normal)
+					.BackgroundColor(FLinearColor(0,0,0,0))
+				]
+			];
+
+		return SNew(SWidgetSwitcher)
+		.WidgetIndex_Lambda(
+			[Value]{
+				const FString ResolvedValue = Value.Get().ToString();
+				int32 Index;
+				return ResolvedValue.FindChar('\n', Index);
+			}
+		)
+		+SWidgetSwitcher::Slot()
+		[
+			SNew(SHorizontalBox)
+			+SHorizontalBox::Slot()
+			.VAlign(VAlign_Center)
+			.AutoWidth()
+			[
+				KeyWidget
+			]
+			+SHorizontalBox::Slot()
+			[
+				ValueWidget
+			]
+		]
+		+SWidgetSwitcher::Slot()
+		[
+			SNew(SHorizontalBox)
+			+SHorizontalBox::Slot()
+			.AutoWidth()
+			[
+				KeyWidget
+			]
+			+SHorizontalBox::Slot()
+			[
+				ValueWidget
+			]
+		];
+	};
+	
+	return SAssignNew(ChangelistInfoWidget, SVerticalBox)
+	.Visibility(EVisibility::Collapsed)
+	+SVerticalBox::Slot()
+	.AutoHeight()
+	[
+		GenerateRow(LOCTEXT("ChangelistAuthor", "Author"), TAttribute<FText>::CreateLambda(
+			[this]{return CurrentChangelistInfo.Author;}))
+	]
+	+SVerticalBox::Slot()
+	.AutoHeight()
+	[
+		GenerateRow(LOCTEXT("ChangelistPath", "Path"), TAttribute<FText>::CreateLambda(
+			[this]{return CurrentChangelistInfo.SharedPath;}))
+	]
+	+SVerticalBox::Slot()
+	.AutoHeight()
+	[
+		GenerateRow(LOCTEXT("ChangelistStatus", "Status"), TAttribute<FText>::CreateLambda(
+			[this]{return CurrentChangelistInfo.Status;}))
+	]
+	+SVerticalBox::Slot()
+	.AutoHeight()
+	[
+		GenerateRow(LOCTEXT("ChangelistDescription", "Description"), TAttribute<FText>::CreateLambda(
+			[this]{return CurrentChangelistInfo.Description;}))
+	];
+}
+
+void SSourceControlReview::LoadChangelist(const FString& Changelist)
+{
+	ChangelistEntriesWidget->ClearChildren();
+	
+	if (!IsLoading())
+	{
+		SetLoading(true);
+	}
+	else
+	{
+		FMessageDialog::Open(EAppMsgType::Ok, NSLOCTEXT("ReviewChangelistTool", "ChangelistError", "Changelist is already loading"));
+		return;
+	}
+
+	ChangelistFileDataMap.Empty();
+
+	//This command runs p4 -describe (or similar for other version controls) to retrieve changelist record information
+	const TSharedRef<FGetChangelistDetails> GetChangelistDetailsCommand = ISourceControlOperation::Create<FGetChangelistDetails>();
+	GetChangelistDetailsCommand->SetChangelistNumber(Changelist);
+
+	ISourceControlModule::Get().GetProvider().Execute(
+		GetChangelistDetailsCommand,
+		EConcurrency::Asynchronous,
+		FSourceControlOperationComplete::CreateRaw(this, &SSourceControlReview::OnChangelistLoadComplete, Changelist));
+}
+
+void SSourceControlReview::OnChangelistLoadComplete(const FSourceControlOperationRef& InOperation, ECommandResult::Type InResult, FString Changelist)
+{
+	const TSharedRef<FGetChangelistDetails, ESPMode::ThreadSafe> Operation = StaticCastSharedRef<FGetChangelistDetails>(InOperation);
+	const TArray<TMap<FString, FString>>& Record = Operation->GetChangelistDetails();
+	
+	if (!IsChangelistRecordValid(Record))
+	{
+		SetLoading(false);
+		return;
+	}
+
+	const TMap<FString, FString>& ChangelistRecord = Record[RecordIndex];
+	
+	//Num files we are going to expect retrieved from source control
+	FilesToLoad = 0;
+	//Num files we've loaded so far
+	FilesLoaded = 0;
+
+	// TODO: @jordan.hoffmann revise for other version controls
+	//Each file in p4 has an index "depotFile0" this is the index that is updated by the while loop to find all the files "depotFile1", "depotFile2" etc...
+	uint32  RecordFileIndex = 0;
+	//String representation of the current file index
+	FString RecordFileIndexStr = LexToString(RecordFileIndex);
+	//The p4 records is the map a file key starts with "depotFile" and is followed by file index 
+	FString RecordFileMapKey = FileDepotKey + RecordFileIndexStr;
+	//The p4 records is the map a revision key starts with "rev" and is followed by file index 
+	FString RecordRevisionMapKey = FileRevisionKey + RecordFileIndexStr;
+	//The p4 records is the map a revision key starts with "action" and is followed by file index 
+	FString RecordActionMapKey = FileActionKey + RecordFileIndexStr;
+
+	SetChangelistInfo(ChangelistRecord);
+	
+	//the loop checks if we have a valid record "depotFile(Index)" in the records to add a file entry
+	while (ChangelistRecord.Contains(RecordFileMapKey) && ChangelistRecord.Contains(RecordRevisionMapKey))
+	{
+		const FString &FileDepotPath = ChangelistRecord[RecordFileMapKey];
+		const FString AssetName = FPaths::GetBaseFilename(FileDepotPath, true);
+		
+		//For each 1 file we are loading 2 revisions so files to load is always incremented by two per file
+		FilesToLoad++; 
+
+		const bool bIsShelved = ChangelistRecord[ChangelistStatusKey] == ChangelistPendingStatusKey;
+		const int32 AssetRevision = FCString::Atoi(*ChangelistRecord[RecordRevisionMapKey]);
+	
+		TSharedPtr<FChangelistFileData> ChangelistFileData = MakeShared<FChangelistFileData>(AssetName, TEXT(""), ChangelistRecord[RecordRevisionMapKey], TEXT(""), TEXT(""));
+
+		ChangelistFileData->ReviewFileDateTime = FDateTime(1970, 1, 1, 0, 0, 0, 0) + FTimespan::FromSeconds(FCString::Atoi(*ChangelistRecord[TimeKey]));
+		ChangelistFileData->ChangelistNum = FCString::Atoi(*Changelist);
+
+		//Determine if we are dealing with submitted or pending changelist
+		ChangelistFileData->ChangelistState = bIsShelved ? EChangelistState::Pending : EChangelistState::Submitted;
+
+		//Building absolute local path is needed to use local file to retrieve file history information from p4 to then show file revision data
+		ChangelistFileData->AssetFilePath = AsAssetPath(ChangelistRecord[RecordFileMapKey]);
+		ChangelistFileData->RelativeFilePath = TrimSharedPath(ChangelistRecord[RecordFileMapKey]);
+
+		SetFileSourceControlAction(ChangelistFileData, ChangelistRecord[RecordActionMapKey]);
+
+		const int32 PreviousAssetRevision = bIsShelved && (ChangelistFileData->FileSourceControlAction == ESourceControlAction::Delete || ChangelistFileData->FileSourceControlAction == ESourceControlAction::Edit) ? AssetRevision : AssetRevision - 1;
+		const FString PreviousAssetRevisionStr = FString::FromInt(PreviousAssetRevision);
+		
+		// retrieve files directly from source control into a temp location
+		TSharedRef<FGetFile> GetFileToReviewCommand = ISourceControlOperation::Create<FGetFile>(Changelist, ChangelistRecord[RecordRevisionMapKey], ChangelistRecord[RecordFileMapKey], bIsShelved);
+
+		TWeakPtr<SSourceControlReview> WeakReviewWidget = SharedThis(this);
+		const auto GetFileCommandResponse = [WeakReviewWidget, ChangelistFileData](const FSourceControlOperationRef& InOperation, ECommandResult::Type InResult, FString* OutFileName)
+		{
+			if (WeakReviewWidget.IsValid())
+			{
+				const TSharedRef<FGetFile, ESPMode::ThreadSafe> Operation = StaticCastSharedRef<FGetFile>(InOperation);
+				*OutFileName = Operation->GetOutPackageFilename();
+				WeakReviewWidget.Pin()->OnGetFileFromSourceControl(ChangelistFileData);
+			}
+		};
+		
+		ISourceControlModule::Get().GetProvider().Execute(GetFileToReviewCommand, EConcurrency::Asynchronous, FSourceControlOperationComplete::CreateLambda(GetFileCommandResponse, &ChangelistFileData->ReviewFileName));
+
+		if (ChangelistFileData->FileSourceControlAction != ESourceControlAction::Add)
+		{
+			FilesToLoad++;
+
+			TSharedRef<FGetFile> GetPreviousFileCommand = ISourceControlOperation::Create<FGetFile>(Changelist, PreviousAssetRevisionStr, ChangelistRecord[RecordFileMapKey], false);
+
+			ISourceControlModule::Get().GetProvider().Execute(GetPreviousFileCommand, EConcurrency::Asynchronous, FSourceControlOperationComplete::CreateLambda(GetFileCommandResponse, &ChangelistFileData->PreviousFileName));
+		}
+
+		//Update map keys with the next file index
+		//When we are doing looking at for example "depotFile0" we build our next key "depotFile1"
+		RecordFileIndex++;
+		RecordFileIndexStr = LexToString(RecordFileIndex);
+		RecordFileMapKey = FileDepotKey + RecordFileIndexStr;
+		RecordRevisionMapKey = FileRevisionKey + RecordFileIndexStr;
+		RecordActionMapKey = FileActionKey + RecordFileIndexStr;
+	}
+
+	//If we have no files to load flip loading bar visibility
+	if (FilesToLoad == 0)
+	{
+		SetLoading(false);
+	}
+}
+
+FReply SSourceControlReview::OnLoadChangelistClicked()
+{
+	LoadChangelist(ChangelistNumWidget->GetText().ToString());
+	return FReply::Handled();
+}
+
+void SSourceControlReview::OnChangelistNumCommitted(const FText& Text, ETextCommit::Type CommitMethod)
+{
+	if (CommitMethod == ETextCommit::OnEnter)
+	{
+		LoadChangelist(Text.ToString());
+	}
+}
+
+void SSourceControlReview::OnGetFileFromSourceControl(TSharedPtr<FChangelistFileData> ChangelistFileData)
+{
+	if (ChangelistFileData && ChangelistFileData->IsDataValidForEntry())
+	{
+		ChangelistFileDataMap.Add(ChangelistFileData->ReviewFileName, ChangelistFileData);
+	}
+
+	FilesLoaded++;
+
+	LoadingProgressBar->SetPercent(FilesToLoad ? static_cast<float>(FilesLoaded) / static_cast<float>(FilesToLoad) : 1.f);
+
+	if (FilesToLoad == FilesLoaded)
+	{
+		const FAssetRegistryModule& AssetRegistryModule = FModuleManager::LoadModuleChecked<FAssetRegistryModule>( TEXT( "AssetRegistry" ) );
+
+		TArray<TSharedPtr<FChangelistFileData>> ChangelistFilesData;
+		ChangelistFileDataMap.GenerateValueArray(ChangelistFilesData);
+		Algo::Sort(ChangelistFilesData, [](const TSharedPtr<FChangelistFileData> &A, const TSharedPtr<FChangelistFileData> &B)
+		{
+			return A->RelativeFilePath < B->RelativeFilePath;
+		});
+		
+		TArray<FString> ChangelistFilePaths;
+		for (const TSharedPtr<FChangelistFileData>& FileData : ChangelistFilesData)
+		{
+			if(FileData)
+			{
+				ChangelistFilePaths.Add(FileData->ReviewFileName);
+			}
+		}
+		
+		AssetRegistryModule.Get().ScanFilesSynchronous(ChangelistFilePaths);
+		
+		for(const TSharedPtr<FChangelistFileData>& FileData : ChangelistFilesData)
+		{
+			AddDiffFile(*FileData);
+		}
+		SetLoading(false);
+	}
+}
+
+void SSourceControlReview::SetLoading(bool bInLoading)
+{
+	bChangelistLoading = bInLoading;
+	
+	// show loading bar and text if we're loading
+	LoadingProgressBar->SetVisibility(bInLoading? EVisibility::Visible : EVisibility::Collapsed);
+	LoadingTextBlock->SetVisibility(bInLoading? EVisibility::Visible : EVisibility::Collapsed);
+
+	// hide CL info if we're loading or there's no changelist
+	ChangelistInfoWidget->SetVisibility((ChangelistNumWidget->GetText().IsEmpty() || bInLoading)? EVisibility::Collapsed : EVisibility::Visible);
+}
+
+bool SSourceControlReview::IsLoading() const
+{
+	return bChangelistLoading;
+}
+
+bool SSourceControlReview::IsChangelistRecordValid(const TArray<TMap<FString, FString>>& InRecord) const
+{
+	if (InRecord.IsEmpty())
+	{
+		FMessageDialog::Open(EAppMsgType::Ok, LOCTEXT("ChangelistNotFoundError", "No record found for this changelist"));
+		return false;
+	}
+	if (InRecord.Num() > 1)
+	{
+		FMessageDialog::Open(EAppMsgType::Ok, LOCTEXT("ChangelistInvalidResponseFormat", "Invalid API response from Source Control"));
+		return false;
+	}
+
+	const TMap<FString, FString>& RecordMap = InRecord[RecordIndex];
+	if (!RecordMap.Contains(ChangelistStatusKey))
+	{
+		FMessageDialog::Open(EAppMsgType::Ok, LOCTEXT("ChangelistMissingStatus", "Changelist is missing status information"));
+		return false;
+	}
+	if (!RecordMap.Contains(AuthorKey))
+	{
+		FMessageDialog::Open(EAppMsgType::Ok, LOCTEXT("ChangelistMissingAuthor", "Changelist is missing author information"));
+		return false;
+	}
+	if (!RecordMap.Contains(DescriptionKey))
+	{
+		FMessageDialog::Open(EAppMsgType::Ok, LOCTEXT("ChangelistMissingDescription", "Changelist is missing description information"));
+		return false;
+	}
+	if (!RecordMap.Contains(TimeKey))
+	{
+		FMessageDialog::Open(EAppMsgType::Ok, LOCTEXT("ChangelistMissingDate", "Changelist is missing date information"));
+		return false;
+	}
+
+	return true;
+}
+
+void SSourceControlReview::SetFileSourceControlAction(TSharedPtr<FChangelistFileData> ChangelistFileData, const FString& SourceControlAction)
+{
+	if (!ChangelistFileData.IsValid())
+	{
+		UE_LOG(LogSourceControl, Error, TEXT("Unable to set source control action information. No changelist file data"));
+		return;
+	}
+
+	if (SourceControlAction == FString(TEXT("add")))
+	{
+		ChangelistFileData->FileSourceControlAction = ESourceControlAction::Add;
+	}
+	else if (SourceControlAction == FString(TEXT("edit")))
+	{
+		ChangelistFileData->FileSourceControlAction = ESourceControlAction::Edit;
+	}
+	else if (SourceControlAction == FString(TEXT("delete")))
+	{
+		ChangelistFileData->FileSourceControlAction = ESourceControlAction::Delete;
+	}
+	else if (SourceControlAction == FString(TEXT("branch")))
+	{
+		ChangelistFileData->FileSourceControlAction = ESourceControlAction::Branch;
+	}
+	else if (SourceControlAction == FString(TEXT("integrate")))
+	{
+		ChangelistFileData->FileSourceControlAction = ESourceControlAction::Integrate;
+	}
+	else
+	{
+		UE_LOG(LogSourceControl, Error, TEXT("Unable to parse source control action information. '%s' diff will not be shown"), *ChangelistFileData->AssetName);
+		ChangelistFileData->FileSourceControlAction = ESourceControlAction::Unset;
+	}
+}
+
+static FString GetSharedBranchPath(const TMap<FString, FString>& InChangelistRecord)
+{
+	FString SharedBranchPath;
+	if (const FString* Found = InChangelistRecord.Find(FileDepotKey + "0"))
+	{
+		SharedBranchPath = *Found;
+	}
+	else
+	{
+		return FString();
+	}
+
+	// TODO: @jordan.hoffmann: Revise for other source controls
+	//Each file in p4 has an index "depotFile0" this is the index that is updated by the while loop to find all the files "depotFile1", "depotFile2" etc...
+	uint32  RecordFileIndex = 1;
+	//The p4 records is the map a file key starts with "depotFile" and is followed by file index 
+	FString RecordFileMapKey = FileDepotKey + LexToString(RecordFileIndex);
+	while (const FString* Found = InChangelistRecord.Find(RecordFileMapKey))
+	{
+		// find starting from the left, find the portion that's shared between both strings
+		int32 TrimIndex = 0;
+		const int32 MaxTrimIndex = FMath::Min(Found->Len(), SharedBranchPath.Len());
+		for(; TrimIndex < MaxTrimIndex; ++TrimIndex)
+		{
+			if (SharedBranchPath[TrimIndex] != (*Found)[TrimIndex])
+			{
+				break;
+			}
+		}
+		// shorten SharedBranchPath to only contain the shared portion
+		SharedBranchPath = SharedBranchPath.Left(TrimIndex);
+
+		// increment to next file path
+		RecordFileMapKey = FileDepotKey + LexToString(++RecordFileIndex);
+	}
+	return SharedBranchPath;
+}
+
+void SSourceControlReview::SetChangelistInfo(const TMap<FString, FString>& InChangelistRecord)
+{
+	CurrentChangelistInfo.Author = FText::FromString(InChangelistRecord[AuthorKey]);
+	CurrentChangelistInfo.Description = FText::FromString(InChangelistRecord[DescriptionKey]);
+	CurrentChangelistInfo.Status = FText::FromString(InChangelistRecord[ChangelistStatusKey]);
+	CurrentChangelistInfo.SharedPath = FText::FromString(GetSharedBranchPath(InChangelistRecord));
+}
+
+void SSourceControlReview::AddDiffFile(const FChangelistFileData& FileData) const
+{
+	ChangelistEntriesWidget->AddSlot()
+	.Padding(0.f, 5.f)
+	[
+		SNew(SSourceControlReviewEntry)
+		.FileData(FileData)
+	];
+}
+
+FString SSourceControlReview::TrimSharedPath(FString FullCLPath) const
+{
+	FullCLPath.RemoveFromStart(*CurrentChangelistInfo.SharedPath.ToString());
+	return FullCLPath;
+}
+
+FString SSourceControlReview::AsAssetPath(const FString& FullCLPath)
+{
+	const FString ProjectName = FApp::GetProjectName();
+	return UKismetSystemLibrary::GetProjectDirectory() / FullCLPath.RightChop(FullCLPath.Find(ProjectName) + ProjectName.Len());
+}
+
+#undef LOCTEXT_NAMESPACE
diff --git a/Engine/Plugins/Editor/ChangelistReview/Source/Private/SSourceControlReview.h b/Engine/Plugins/Editor/ChangelistReview/Source/Private/SSourceControlReview.h
new file mode 100644
--- /dev/null
+++ b/Engine/Plugins/Editor/ChangelistReview/Source/Private/SSourceControlReview.h
@@ -0,0 +1,184 @@
+// Copyright Epic Games, Inc. All Rights Reserved.
+
+#pragma once
+#include "ISourceControlOperation.h"
+#include "ISourceControlProvider.h"
+#include "Widgets/SCompoundWidget.h"
+#include "Widgets/Text/STextBlock.h"
+#include "Widgets/SBoxPanel.h"
+
+class SChangelistEditableText;
+class SProgressBar;
+
+namespace SourceControlReview
+{
+	enum class ESourceControlAction : uint8
+	{
+		Add,
+		Edit,
+		Delete,
+		Branch,
+		Integrate,
+		Unset,
+
+		// keep this last
+		ActionCount
+	};
+
+	enum class EChangelistState : uint8
+	{
+		Submitted,
+		Pending
+	};
+
+	struct FChangelistFileData
+	{
+		FChangelistFileData()
+		{
+		}
+
+		FChangelistFileData(const FString& InAssetName, const FString& InReviewFilePkgName,
+							const FString& InReviewFileRevisionNum, const FString& InPreviousFilePkgName,
+							const FString& InPreviousFileRevisionNum)
+			: AssetName(InAssetName)
+			, ReviewFileName(InReviewFilePkgName)
+			, ReviewFileRevisionNum(InReviewFileRevisionNum)
+			, PreviousFileName(InPreviousFilePkgName)
+			, PreviousFileRevisionNum(InPreviousFileRevisionNum)
+		{
+		}
+
+		bool IsDataValidForEntry() const
+		{
+			return FileSourceControlAction != ESourceControlAction::Unset && !ReviewFileName.IsEmpty() &&
+				(!PreviousFileName.IsEmpty() || FileSourceControlAction == ESourceControlAction::Add);
+		}
+
+		FString AssetName;
+
+		FString ReviewFileName;
+
+		FString ReviewFileRevisionNum;
+
+		FDateTime ReviewFileDateTime;
+
+		FString PreviousFileName;
+
+		FString PreviousFileRevisionNum;
+
+		FString RelativeFilePath;
+
+		FString AssetFilePath;
+
+		int32 ChangelistNum = INDEX_NONE;
+
+		EChangelistState ChangelistState = EChangelistState::Submitted;
+
+		ESourceControlAction FileSourceControlAction = ESourceControlAction::Unset;
+	};
+}
+
+
+/**
+ * Used to select a changelist and diff it's changes
+ */
+class SSourceControlReview : public SCompoundWidget
+{
+public:
+	using ESourceControlAction = SourceControlReview::ESourceControlAction;
+	using EChangelistState = SourceControlReview::EChangelistState;
+	using FChangelistFileData = SourceControlReview::FChangelistFileData;
+	
+	SLATE_BEGIN_ARGS(SSourceControlReview) {}
+	SLATE_END_ARGS()
+	
+	/** Constructs the widget */
+	void Construct(const FArguments& InArgs);
+
+	/**
+	 * Pulls up changelist record from source control
+	 * @param Changelist Changelist number to get
+	 */
+	void LoadChangelist(const FString& Changelist);
+	
+private:
+	void OnChangelistLoadComplete(const FSourceControlOperationRef& InOperation, ECommandResult::Type InResult, FString Changelist);
+	FReply OnLoadChangelistClicked();
+	void OnChangelistNumCommitted(const FText& Text, ETextCommit::Type CommitMethod);
+	
+	/**
+	 * Called when file is retrieved from
+	 * @param ChangelistFileData carried through worker to be filled with information about retrieved file from source control
+	 */
+	void OnGetFileFromSourceControl(TSharedPtr<FChangelistFileData> ChangelistFileData);
+
+	/**
+	 * Sets bChangelistLoading and then fires the BP_SetLoading event
+	 * @param bInLoading when true flips visibility on loading bar and is also blocking from loading multiple changelists at the same time until loading is finished
+	 */
+	void SetLoading(bool bInLoading);
+	
+	/**
+	 * True if the tool is currently loading a changelist
+	 */
+	bool IsLoading() const;
+
+	/**
+	 * Checks changelist information to make sure it's valid for the diff tool
+	 * @param InRecord changelist record
+	 */
+	bool IsChangelistRecordValid(const TArray<TMap<FString, FString>>& InRecord) const;
+
+	/**
+	 * Sets the source control action type for the loaded file
+	 * @param ChangelistFileData carried through worker to be filled with information about retrieved file from source control
+	 * @param SourceControlAction source action operation type retrieved from source control
+	 */
+	static void SetFileSourceControlAction(TSharedPtr<FChangelistFileData> ChangelistFileData,
+	                                       const FString& SourceControlAction);
+
+	/**
+	 * Sets changelist information Author, Description, Status and Depot name that changelist is going to
+	 */
+	void SetChangelistInfo(const TMap<FString, FString>& InChangelistRecord);
+
+	/**
+	 * Adds a SSourceControlReviewEntry to the GUI
+	 * @param FileData The File being added to the list
+	 */
+	void AddDiffFile(const FChangelistFileData& FileData) const;
+
+	/**
+	 * Removes CurrentChangelistInfo.SharedPath from the beginning of FullCLPath
+	 */
+	FString TrimSharedPath(FString FullCLPath) const;
+
+	/**
+	 * Removes the game directory from the beginning of the FullCLPath
+	 */
+	static FString AsAssetPath(const FString& FullCLPath);
+
+	// used for asynchronous changelist loading
+	bool bChangelistLoading = false;
+	uint32 FilesToLoad = 0;
+	uint32 FilesLoaded = 0;
+	TMap<FString, TSharedPtr<FChangelistFileData>> ChangelistFileDataMap;
+
+	// builds a widget that displays a changelist's Author, Path, Status, and Description
+	TSharedRef<SWidget> GetChangelistInfoWidget();
+	
+	TSharedPtr<SChangelistEditableText> ChangelistNumWidget;
+	TSharedPtr<STextBlock> LoadingTextBlock;
+	TSharedPtr<SProgressBar> LoadingProgressBar;
+	TSharedPtr<SWidget> ChangelistInfoWidget;
+	TSharedPtr<SVerticalBox> ChangelistEntriesWidget;
+
+	// Info about the current chagnelist
+	struct FChangelistInfo
+	{
+		FText Author;
+		FText SharedPath;
+		FText Status;
+		FText Description;
+	} CurrentChangelistInfo;
+};
diff --git a/Engine/Plugins/Editor/ChangelistReview/Source/Private/SSourceControlReviewEntry.cpp b/Engine/Plugins/Editor/ChangelistReview/Source/Private/SSourceControlReviewEntry.cpp
new file mode 100644
--- /dev/null
+++ b/Engine/Plugins/Editor/ChangelistReview/Source/Private/SSourceControlReviewEntry.cpp
@@ -0,0 +1,383 @@
+// Copyright Epic Games, Inc. All Rights Reserved.
+
+#include "SSourceControlReviewEntry.h"
+
+#include "AssetToolsModule.h"
+#include "ClassIconFinder.h"
+#include "SourceControlHelpers.h"
+#include "AssetRegistry/AssetRegistryModule.h"
+#include "Engine/Font.h"
+#include "Kismet2/KismetEditorUtilities.h"
+#include "Misc/FileHelper.h"
+#include "Misc/MessageDialog.h"
+#include "Settings/EditorLoadingSavingSettings.h"
+#include "Styling/StarshipCoreStyle.h"
+#include "Widgets/Text/STextBlock.h"
+#include "Widgets/Layout/SBox.h"
+#include "Widgets/Input/SButton.h"
+
+#define LOCTEXT_NAMESPACE "SourceControlReviewEntry"
+
+namespace ReviewEntryConsts
+{
+	static const FString TempFolder = TEXT("/Temp/");
+}
+
+void SSourceControlReviewEntry::Construct(const FArguments& InArgs)
+{
+	static const UFont* Font = LoadObject<UFont>(nullptr, TEXT("/Game/UI/Foundation/Fonts/NotoSans.NotoSans"));
+	
+	ChangelistFileData = InArgs._FileData;
+	ChildSlot
+	[
+		SNew(SHorizontalBox)
+		
+		+SHorizontalBox::Slot()
+		.AutoWidth()
+		.Padding(10.f, 0.f)
+		.HAlign(HAlign_Center)
+		.VAlign(VAlign_Center)
+		[
+			SAssignNew(SourceActionIcon, SImage)
+			.Image(this, &SSourceControlReviewEntry::GetSourceControlIcon)
+		]
+		
+		+SHorizontalBox::Slot()
+		.FillWidth(1.f)
+		.Padding(0.f, 0.f, 5.f, 0.f)
+		.VAlign(VAlign_Center)
+		[
+			SNew(SVerticalBox)
+			+SVerticalBox::Slot()
+			.AutoHeight()
+			[
+				SNew(STextBlock)
+				.Text(this, &SSourceControlReviewEntry::GetAssetNameText)
+				.Font(FStyleFonts::Get().Large)
+			]
+			+SVerticalBox::Slot()
+			.AutoHeight()
+			[
+				SNew(STextBlock)
+				.Text(this, &SSourceControlReviewEntry::GetLocalAssetPathText)
+				.Font(FStyleFonts::Get().Small)
+			]
+		]
+		
+		+SHorizontalBox::Slot()
+		.AutoWidth()
+		.Padding(10.f, 0.f)
+		.HAlign(HAlign_Center)
+		.VAlign(VAlign_Center)
+		[
+			SNew(SBox)
+			.WidthOverride(32.f)
+			.HeightOverride(32.f)
+			[
+				SAssignNew(AssetTypeIcon, SImage)
+				.Image(GetAssetTypeIcon())
+			]
+		]
+		
+		+SHorizontalBox::Slot()
+		.AutoWidth()
+		.HAlign(HAlign_Center)
+		.VAlign(VAlign_Center)
+		[
+			SAssignNew(FileDeletedTextBlock, STextBlock)
+		]
+		
+		+SHorizontalBox::Slot()
+		.FillWidth(1.f)
+		.Padding(10.f, 0.f)
+		.VAlign(VAlign_Center)
+		[
+			SAssignNew(ReviewInputsBox, SHorizontalBox)
+			+SHorizontalBox::Slot()
+			.FillWidth(1.f)
+			.Padding(0.f, 0.f, 5.f, 0.f)
+			[
+				SNew(SButton)
+				.VAlign(VAlign_Center)
+				[
+					SAssignNew(ViewDiffButtonText, STextBlock)
+					.Justification(ETextJustify::Center)
+					.Text(LOCTEXT("ViewDiffButton", "View Diff"))
+				]
+				.OnClicked(this, &SSourceControlReviewEntry::OnDiffClicked)
+			]
+			+SHorizontalBox::Slot()
+			.FillWidth(1.f)
+			.Padding(0.f, 0.f, 5.f, 0.f)
+			[
+				SNew(SButton)
+				.IsEnabled(this, &SSourceControlReviewEntry::CanBrowseToAsset)
+				.ContentPadding(FMargin(8.f, 4.f))
+				.OnClicked(this, &SSourceControlReviewEntry::OnBrowseToAssetClicked)
+				[
+					SNew(SHorizontalBox)
+					+SHorizontalBox::Slot()
+					.AutoWidth()
+					.VAlign(VAlign_Center)
+					.HAlign(HAlign_Left)
+					[
+						SNew(SImage)
+						.Image(FAppStyle::GetBrush("Icons.Search"))
+					]
+					+SHorizontalBox::Slot()
+					.FillWidth(1.f)
+					.VAlign(VAlign_Center)
+					.HAlign(HAlign_Center)
+					[
+						SNew(STextBlock)
+						.Justification(ETextJustify::Center)
+						.Text(LOCTEXT("BrowseToAssetButton", "Browse To Asset"))
+					]
+				]
+			]
+		]
+	];
+
+	// figure out how this asset diffs, and bind it to this->DiffMethod
+	TryBindDiffMethod();
+}
+
+void SSourceControlReviewEntry::SetEntryData(const FChangelistFileData& InChangelistFileData)
+{
+	ChangelistFileData = InChangelistFileData;
+	AssetTypeIcon->SetImage(GetAssetTypeIcon());
+
+	// if asset changed, we might diff differently. rebind the diff method.
+	DiffMethod.Unbind();
+	TryBindDiffMethod();
+}
+
+FReply SSourceControlReviewEntry::OnDiffClicked() const
+{
+	DiffMethod.Execute();
+	return FReply::Handled();
+}
+
+FReply SSourceControlReviewEntry::OnBrowseToAssetClicked() const
+{
+	TArray<FAssetData> Assets;
+	USourceControlHelpers::GetAssetData(ChangelistFileData.AssetFilePath, Assets);
+	const FAssetToolsModule& AssetToolsModule = FModuleManager::LoadModuleChecked<FAssetToolsModule>(TEXT("AssetTools"));
+	AssetToolsModule.Get().SyncBrowserToAssets(Assets);
+	return FReply::Handled();
+}
+
+bool SSourceControlReviewEntry::CanBrowseToAsset() const
+{
+	TArray<FAssetData> Assets;
+	return USourceControlHelpers::GetAssetData(ChangelistFileData.AssetFilePath, Assets);
+}
+
+void SSourceControlReviewEntry::TryBindDiffMethod()
+{
+	if (!DiffMethod.IsBound())
+	{
+		TryBindUAssetDiff();
+	}
+	if (!DiffMethod.IsBound())
+	{
+		TryBindTextDiff();
+	}
+}
+
+bool SSourceControlReviewEntry::CanDiff() const
+{
+	return DiffMethod.IsBound();
+}
+
+void SSourceControlReviewEntry::TryBindUAssetDiff()
+{
+	if (UPackage* ReviewFilePkg = LoadPackage(nullptr, *ChangelistFileData.ReviewFileName, LOAD_ForDiff | LOAD_DisableCompileOnLoad | LOAD_DisableEngineVersionChecks))
+	{
+		if (UObject* ReviewAsset = FindObject<UObject>(ReviewFilePkg, *ChangelistFileData.AssetName))
+		{
+			UPackage* PreviousFilePkg = LoadPackage(nullptr, *ChangelistFileData.PreviousFileName, LOAD_ForDiff | LOAD_DisableCompileOnLoad | LOAD_DisableEngineVersionChecks);
+
+			if (UObject* PreviousAsset = PreviousFilePkg ? FindObject<UObject>(PreviousFilePkg, *ChangelistFileData.AssetName) : nullptr)
+			{
+				DiffMethod.BindLambda([this, PreviousAsset, ReviewAsset]
+				{
+					FAssetToolsModule::GetModule().Get().DiffAssets(
+						PreviousAsset,
+						ReviewAsset,
+						GetPreviousFileRevisionInfo(),
+						GetReviewFileRevisionInfo()
+					);
+				});
+				
+			}
+			else
+			{
+				if (const UBlueprint* BlueprintReviewAsset = Cast<UBlueprint>(ReviewAsset))
+				{
+					DiffMethod.BindLambda([this, BlueprintReviewAsset, ReviewAsset]
+					{
+						FAssetToolsModule::GetModule().Get().DiffAssets(
+							GetOrCreateBlueprintForDiff(BlueprintReviewAsset->GeneratedClass, BlueprintReviewAsset->BlueprintType),
+							ReviewAsset,
+							GetPreviousFileRevisionInfo(),
+							GetReviewFileRevisionInfo()
+						);
+					});
+				}
+				else
+				{
+					//In case when file is not blueprint we are falling back on trying to create a UObject 
+					const FName EmptyObjectName = ReviewAsset->GetFName();
+					const UClass* EmptyObjectAssetClass = ReviewAsset->GetClass();
+					
+					if (UObject* EmptyObject = NewObject<UObject>(ReviewAsset, EmptyObjectAssetClass, EmptyObjectName, ReviewAsset->GetFlags()))
+					{
+						DiffMethod.BindLambda([this, EmptyObject, ReviewAsset]
+						{
+							FAssetToolsModule::GetModule().Get().DiffAssets(
+								EmptyObject,
+								ReviewAsset,
+								GetPreviousFileRevisionInfo(),
+								GetReviewFileRevisionInfo()
+							);
+						});
+					}
+				}
+			}
+		}
+	}
+}
+
+void SSourceControlReviewEntry::TryBindTextDiff()
+{
+	DiffMethod.BindLambda([this]
+	{
+		const FString& DiffCommand = GetDefault<UEditorLoadingSavingSettings>()->TextDiffToolPath.FilePath;
+		const FAssetToolsModule& AssetToolsModule = FAssetToolsModule::GetModule();
+		AssetToolsModule.Get().CreateDiffProcess(DiffCommand, ChangelistFileData.PreviousFileName, ChangelistFileData.ReviewFileName);
+	});
+}
+
+FRevisionInfo SSourceControlReviewEntry::GetReviewFileRevisionInfo() const
+{
+	FRevisionInfo ReviewFileRevisionInfo;
+	ReviewFileRevisionInfo.Changelist = ChangelistFileData.ChangelistNum;
+	ReviewFileRevisionInfo.Date = ChangelistFileData.ReviewFileDateTime;
+
+	if (ChangelistFileData.ChangelistState == EChangelistState::Pending)
+	{
+		ReviewFileRevisionInfo.Revision = TEXT("Pending");
+	}
+	else
+	{
+		ReviewFileRevisionInfo.Revision = ChangelistFileData.ReviewFileRevisionNum;
+	}
+
+	return ReviewFileRevisionInfo;
+}
+
+FRevisionInfo SSourceControlReviewEntry::GetPreviousFileRevisionInfo() const
+{
+	//We need to have valid revision data for some DiffAssets implementations (Although for now we don't have full data on previous file we are showing correct previous revision information)
+	FRevisionInfo PreviousFileRevisionInfo;
+	PreviousFileRevisionInfo.Revision = ChangelistFileData.PreviousFileRevisionNum.IsEmpty() ? TEXT("0") : ChangelistFileData.PreviousFileRevisionNum;
+	return PreviousFileRevisionInfo;
+}
+
+const FSlateBrush* SSourceControlReviewEntry::GetSourceControlIcon() const
+{
+	// setup lookup table for all the brushes so we don't have to re-call FAppStyle::Get().GetBrush every frame
+	static TArray<const FSlateBrush*> Brushes = []()
+	{
+		TArray<const FSlateBrush*> Temp;
+		Temp.AddZeroed(static_cast<uint8>(ESourceControlAction::ActionCount));
+		
+		//Source control images can be found at Engine\Content\Slate\Starship\SourceControl
+		Temp[static_cast<uint8>(ESourceControlAction::Add)] = FAppStyle::Get().GetBrush(TEXT("SourceControl.Add"));
+		Temp[static_cast<uint8>(ESourceControlAction::Edit)] = FAppStyle::Get().GetBrush(TEXT("SourceControl.Edit"));
+		Temp[static_cast<uint8>(ESourceControlAction::Delete)] = FAppStyle::Get().GetBrush(TEXT("SourceControl.Delete"));
+		Temp[static_cast<uint8>(ESourceControlAction::Branch)] = FAppStyle::Get().GetBrush(TEXT("SourceControl.Branch"));
+		Temp[static_cast<uint8>(ESourceControlAction::Integrate)] = FAppStyle::Get().GetBrush(TEXT("SourceControl.Integrate"));
+		Temp[static_cast<uint8>(ESourceControlAction::Unset)] = FAppStyle::Get().GetBrush(TEXT("SourceControl.Edit"));
+		
+		return MoveTemp(Temp);
+	}();
+
+	uint8 Index = static_cast<uint8>(ChangelistFileData.FileSourceControlAction);
+	if (!Brushes.IsValidIndex(Index))
+	{
+		Index = static_cast<uint8>(ESourceControlAction::Unset);
+	}
+	return Brushes[Index];
+}
+
+const FSlateBrush* SSourceControlReviewEntry::GetAssetTypeIcon() const
+{
+	const FAssetRegistryModule& AssetRegistryModule = FModuleManager::LoadModuleChecked<FAssetRegistryModule>( TEXT( "AssetRegistry" ) );
+	
+	FString TempPackageName;
+	if (FPackageName::TryConvertFilenameToLongPackageName(ChangelistFileData.ReviewFileName, TempPackageName))
+	{
+		TArray<FAssetData> OutAssetData;
+		AssetRegistryModule.Get().GetAssetsByPackageName( *TempPackageName, OutAssetData );
+		
+		if (OutAssetData.Num() > 0)
+		{
+			if (const FSlateBrush* FileTypeBrush = FClassIconFinder::FindThumbnailForClass(FClassIconFinder::GetIconClassForAssetData(OutAssetData[0])))
+			{
+				return FileTypeBrush;
+			}
+		}
+	}
+	
+	return FAppStyle::Get().GetBrush(TEXT("NoBrush"));
+}
+
+FText SSourceControlReviewEntry::GetAssetNameText() const
+{
+	return FText::FromString(ChangelistFileData.AssetName);
+}
+
+FText SSourceControlReviewEntry::GetLocalAssetPathText() const
+{
+	return FText::FromString(ChangelistFileData.RelativeFilePath);
+}
+
+const FString& SSourceControlReviewEntry::GetSearchableString() const
+{
+	return ChangelistFileData.AssetName;
+}
+
+UBlueprint* SSourceControlReviewEntry::GetOrCreateBlueprintForDiff(UClass* InGeneratedClass, EBlueprintType InBlueprintType) const
+{
+	if (!InGeneratedClass || !InGeneratedClass->ClassGeneratedBy)
+	{
+		FMessageDialog::Open(EAppMsgType::Ok, LOCTEXT("ReviewChangelistEntry", "Unable to show the diff for added file because generated class is not valid"));
+		return nullptr;
+	}
+	
+	const FString PackageName = ReviewEntryConsts::TempFolder / InGeneratedClass->GetName();
+	const FName BPName = InGeneratedClass->GetFName();
+	
+	UPackage* BlueprintPackage = CreatePackage(*PackageName);
+	BlueprintPackage->SetPackageFlags(PKG_ForDiffing);
+	
+	if (UBlueprint* ExistingBlueprint = FindObject<UBlueprint>(BlueprintPackage, *InGeneratedClass->GetName()))
+	{
+		return ExistingBlueprint;
+	}
+
+	UBlueprint* BlueprintObject = FKismetEditorUtilities::CreateBlueprint(InGeneratedClass, BlueprintPackage, BPName, InBlueprintType,
+																		  InGeneratedClass->ClassGeneratedBy->GetClass(),
+																		  InGeneratedClass->GetClass(), FName("DiffToolActions"));
+	
+	if (BlueprintObject)
+	{
+		FAssetRegistryModule::AssetCreated(BlueprintObject);
+	}
+	
+	return BlueprintObject;
+}
+
+#undef LOCTEXT_NAMESPACE
diff --git a/Engine/Plugins/Editor/ChangelistReview/Source/Private/SSourceControlReviewEntry.h b/Engine/Plugins/Editor/ChangelistReview/Source/Private/SSourceControlReviewEntry.h
new file mode 100644
--- /dev/null
+++ b/Engine/Plugins/Editor/ChangelistReview/Source/Private/SSourceControlReviewEntry.h
@@ -0,0 +1,122 @@
+// Copyright Epic Games, Inc. All Rights Reserved.
+
+#pragma once
+#include "IAssetTypeActions.h"
+#include "SSourceControlReview.h"
+#include "Engine/Blueprint.h"
+#include "Widgets/Images/SImage.h"
+#include "Widgets/DeclarativeSyntaxSupport.h"
+#include "Widgets/SCompoundWidget.h"
+
+class SSourceControlReviewEntry : public SCompoundWidget
+{
+public:
+	using ESourceControlAction = SourceControlReview::ESourceControlAction;
+	using EChangelistState = SourceControlReview::EChangelistState;
+	using FChangelistFileData = SourceControlReview::FChangelistFileData;
+	
+	SLATE_BEGIN_ARGS(SSourceControlReviewEntry) {}
+		SLATE_ARGUMENT(FChangelistFileData, FileData);
+	SLATE_END_ARGS()
+	
+	/** Constructs the widget */
+	void Construct(const FArguments& InArgs);
+	
+	/**
+	 * Sets needed data to show diff and browse to asset
+	 * @param InChangelistFileData changelist entry information
+	 */
+	void SetEntryData(const FChangelistFileData& InChangelistFileData);
+
+	/**
+	 * Hooks into built in engine diff functionality to show differences between supplied file revisions
+	 */
+	FReply OnDiffClicked() const;
+
+	/**
+	 * Uses engine sync to browse to asset functionality to navigate the content window to file
+	 */
+	FReply OnBrowseToAssetClicked() const;
+
+private:
+	/**
+	 * Checks if we have a valid local file for this entry to sync content window to
+	 */
+	bool CanBrowseToAsset() const;
+
+	/**
+	 * if this file is diffable, bind the diff method to this->DiffMethod
+	 */
+	void TryBindDiffMethod();
+
+	/**
+	 * Checks if this asset is diffable
+	 */
+	bool CanDiff() const;
+
+	/**
+	 * Runs logic to do the diff for UAsset objects
+	 */
+	void TryBindUAssetDiff();
+
+	/**
+	 * Runs logic to do the diff for text based objects
+	 */
+	void TryBindTextDiff();
+
+	/**
+	 * Gets revision information for review file
+	 */
+	FRevisionInfo GetReviewFileRevisionInfo() const;
+
+	/**
+	 * Gets revision information for previous file (Currently Limited to only revision num)
+	 */
+	FRevisionInfo GetPreviousFileRevisionInfo() const;
+
+	/**
+	 * return an icon reflecting this entry's source control action (i.e. add, edit, delete, etc...)
+	 */
+	const FSlateBrush* GetSourceControlIcon() const;
+	
+	/**
+	 * return an icon that indicates the type of the asset being displayed
+	 */
+	const FSlateBrush* GetAssetTypeIcon() const;
+
+	/**
+	 * Name of the asset being displayed
+	 */
+	FText GetAssetNameText() const;
+
+	/**
+	 * Path of the asset being displayed (relative to the changelist path)
+	 */
+	FText GetLocalAssetPathText() const;
+
+	/**
+	 * Returns asset name to be used for search/filter 
+	 */
+	const FString& GetSearchableString() const;
+	
+	/**
+	* Gets or creates a new blueprint of the same class as file about to be diffed to be used for diffing class defaults when there is no previous revision
+	*/
+	UBlueprint* GetOrCreateBlueprintForDiff(UClass* InGeneratedClass, EBlueprintType InBlueprintType) const;
+
+private:
+	// File associated with this entry
+	FChangelistFileData ChangelistFileData;
+	
+	// Called to diff this file against it's previous revision
+	DECLARE_DELEGATE(FDiffMethod);
+	FDiffMethod DiffMethod;
+
+	// Sub-Widgets
+	TSharedPtr<SImage> SourceActionIcon;
+	TSharedPtr<SImage> AssetTypeIcon;
+	TSharedPtr<STextBlock> FileDeletedTextBlock;
+	TSharedPtr<SHorizontalBox> ReviewInputsBox;
+	TSharedPtr<SImage> ViewedEntryIcon;
+	TSharedPtr<STextBlock> ViewDiffButtonText;
+};
diff --git a/Engine/Plugins/Editor/ChangelistReview/Source/Public/ChangelistReviewModule.h b/Engine/Plugins/Editor/ChangelistReview/Source/Public/ChangelistReviewModule.h
new file mode 100644
--- /dev/null
+++ b/Engine/Plugins/Editor/ChangelistReview/Source/Public/ChangelistReviewModule.h
@@ -0,0 +1,26 @@
+// Copyright Epic Games, Inc. All Rights Reserved.
+
+#pragma once
+
+//#include "SSourceControlReview.h"
+#include "Modules/ModuleManager.h"
+#include "UObject/WeakObjectPtr.h"
+#include "UObject/StrongObjectPtr.h"
+#include "Widgets/Docking/SDockTab.h"
+
+class UContentBrowserAliasDataSource;
+
+class FChangelistReviewModule : public FDefaultModuleImpl
+{
+public:
+	virtual void StartupModule() override;
+	virtual void ShutdownModule() override;
+	void ShowReviewTab();
+	bool CanShowReviewTab() const;
+private:
+	TSharedRef<SDockTab> CreateReviewTab(const FSpawnTabArgs& Args);
+	TSharedPtr<SWidget> CreateReviewUI();
+	
+	TWeakPtr<SDockTab> ReviewTab;
+	//TWeakPtr<SSourceControlReview> ReviewWidget;
+};
