[BugFix] RecordIndex declaration shadows a variable in namespace '(anonymous)'

diff --git a/Engine/Plugins/Editor/ChangelistReview/Source/Private/SSourceControlReview.cpp b/Engine/Plugins/Editor/ChangelistReview/Source/Private/SSourceControlReview.cpp
--- a/Engine/Plugins/Editor/ChangelistReview/Source/Private/SSourceControlReview.cpp
+++ b/Engine/Plugins/Editor/ChangelistReview/Source/Private/SSourceControlReview.cpp
@@ -21,7 +21,7 @@
 
 #define LOCTEXT_NAMESPACE "SourceControlReview"
 
-namespace
+namespace ReviewHelpers
 {
 	const FText EnterChangelistText(LOCTEXT("EnterChangelistText", "Enter a changelist number to diff:"));
 	const FText EnterChangelistTooltip(LOCTEXT("EnterChangelistTooltip", "Enter changelist"));
@@ -60,7 +60,7 @@ void SSourceControlReview::Construct(const FArguments& InArgs)
 				.AutoWidth()
 				[
 					SNew(STextBlock)
-					.Text(EnterChangelistText)
+					.Text(ReviewHelpers::EnterChangelistText)
 					.Font(FStyleFonts::Get().LargeBold)
 				]
 				+SHorizontalBox::Slot()
@@ -69,7 +69,7 @@ void SSourceControlReview::Construct(const FArguments& InArgs)
 				.AutoWidth()
 				[
 					SNew(SBorder)
-					.ToolTipText(EnterChangelistTooltip)
+					.ToolTipText(ReviewHelpers::EnterChangelistTooltip)
 					[
 						SAssignNew(ChangelistNumWidget, SChangelistEditableText)
 						.Font(FStyleFonts::Get().Large)
@@ -87,7 +87,7 @@ void SSourceControlReview::Construct(const FArguments& InArgs)
 					.OnClicked(this, &SSourceControlReview::OnLoadChangelistClicked)
 					[
 						SNew(STextBlock)
-						.Text(LoadChangelistText)
+						.Text(ReviewHelpers::LoadChangelistText)
 						.Font(FStyleFonts::Get().LargeBold)
 					]
 				]
@@ -100,7 +100,7 @@ void SSourceControlReview::Construct(const FArguments& InArgs)
 			[
 				SAssignNew(LoadingTextBlock, STextBlock)
 				.Visibility(EVisibility::Collapsed)
-				.Text(LoadingText)
+				.Text(ReviewHelpers::LoadingText)
 				.Font(FStyleFonts::Get().Large)
 			]
 			+SVerticalBox::Slot()
@@ -253,7 +253,7 @@ void SSourceControlReview::OnChangelistLoadComplete(const FSourceControlOperatio
 		return;
 	}
 
-	const TMap<FString, FString>& ChangelistRecord = Record[RecordIndex];
+	const TMap<FString, FString>& ChangelistRecord = Record[ReviewHelpers::RecordIndex];
 	
 	//Num files we are going to expect retrieved from source control
 	FilesToLoad = 0;
@@ -266,11 +266,11 @@ void SSourceControlReview::OnChangelistLoadComplete(const FSourceControlOperatio
 	//String representation of the current file index
 	FString RecordFileIndexStr = LexToString(RecordFileIndex);
 	//The p4 records is the map a file key starts with "depotFile" and is followed by file index 
-	FString RecordFileMapKey = FileDepotKey + RecordFileIndexStr;
+	FString RecordFileMapKey = ReviewHelpers::FileDepotKey + RecordFileIndexStr;
 	//The p4 records is the map a revision key starts with "rev" and is followed by file index 
-	FString RecordRevisionMapKey = FileRevisionKey + RecordFileIndexStr;
+	FString RecordRevisionMapKey = ReviewHelpers::FileRevisionKey + RecordFileIndexStr;
 	//The p4 records is the map a revision key starts with "action" and is followed by file index 
-	FString RecordActionMapKey = FileActionKey + RecordFileIndexStr;
+	FString RecordActionMapKey = ReviewHelpers::FileActionKey + RecordFileIndexStr;
 
 	SetChangelistInfo(ChangelistRecord);
 	
@@ -283,12 +283,12 @@ void SSourceControlReview::OnChangelistLoadComplete(const FSourceControlOperatio
 		//For each 1 file we are loading 2 revisions so files to load is always incremented by two per file
 		FilesToLoad++; 
 
-		const bool bIsShelved = ChangelistRecord[ChangelistStatusKey] == ChangelistPendingStatusKey;
+		const bool bIsShelved = ChangelistRecord[ReviewHelpers::ChangelistStatusKey] == ReviewHelpers::ChangelistPendingStatusKey;
 		const int32 AssetRevision = FCString::Atoi(*ChangelistRecord[RecordRevisionMapKey]);
 	
 		TSharedPtr<FChangelistFileData> ChangelistFileData = MakeShared<FChangelistFileData>(AssetName, TEXT(""), ChangelistRecord[RecordRevisionMapKey], TEXT(""), TEXT(""));
 
-		ChangelistFileData->ReviewFileDateTime = FDateTime(1970, 1, 1, 0, 0, 0, 0) + FTimespan::FromSeconds(FCString::Atoi(*ChangelistRecord[TimeKey]));
+		ChangelistFileData->ReviewFileDateTime = FDateTime(1970, 1, 1, 0, 0, 0, 0) + FTimespan::FromSeconds(FCString::Atoi(*ChangelistRecord[ReviewHelpers::TimeKey]));
 		ChangelistFileData->ChangelistNum = FCString::Atoi(*Changelist);
 
 		//Determine if we are dealing with submitted or pending changelist
@@ -332,9 +332,9 @@ void SSourceControlReview::OnChangelistLoadComplete(const FSourceControlOperatio
 		//When we are doing looking at for example "depotFile0" we build our next key "depotFile1"
 		RecordFileIndex++;
 		RecordFileIndexStr = LexToString(RecordFileIndex);
-		RecordFileMapKey = FileDepotKey + RecordFileIndexStr;
-		RecordRevisionMapKey = FileRevisionKey + RecordFileIndexStr;
-		RecordActionMapKey = FileActionKey + RecordFileIndexStr;
+		RecordFileMapKey = ReviewHelpers::FileDepotKey + RecordFileIndexStr;
+		RecordRevisionMapKey = ReviewHelpers::FileRevisionKey + RecordFileIndexStr;
+		RecordActionMapKey = ReviewHelpers::FileActionKey + RecordFileIndexStr;
 	}
 
 	//If we have no files to load flip loading bar visibility
@@ -429,23 +429,23 @@ bool SSourceControlReview::IsChangelistRecordValid(const TArray<TMap<FString, FS
 		return false;
 	}
 
-	const TMap<FString, FString>& RecordMap = InRecord[RecordIndex];
-	if (!RecordMap.Contains(ChangelistStatusKey))
+	const TMap<FString, FString>& RecordMap = InRecord[ReviewHelpers::RecordIndex];
+	if (!RecordMap.Contains(ReviewHelpers::ChangelistStatusKey))
 	{
 		FMessageDialog::Open(EAppMsgType::Ok, LOCTEXT("ChangelistMissingStatus", "Changelist is missing status information"));
 		return false;
 	}
-	if (!RecordMap.Contains(AuthorKey))
+	if (!RecordMap.Contains(ReviewHelpers::AuthorKey))
 	{
 		FMessageDialog::Open(EAppMsgType::Ok, LOCTEXT("ChangelistMissingAuthor", "Changelist is missing author information"));
 		return false;
 	}
-	if (!RecordMap.Contains(DescriptionKey))
+	if (!RecordMap.Contains(ReviewHelpers::DescriptionKey))
 	{
 		FMessageDialog::Open(EAppMsgType::Ok, LOCTEXT("ChangelistMissingDescription", "Changelist is missing description information"));
 		return false;
 	}
-	if (!RecordMap.Contains(TimeKey))
+	if (!RecordMap.Contains(ReviewHelpers::TimeKey))
 	{
 		FMessageDialog::Open(EAppMsgType::Ok, LOCTEXT("ChangelistMissingDate", "Changelist is missing date information"));
 		return false;
@@ -492,7 +492,7 @@ void SSourceControlReview::SetFileSourceControlAction(TSharedPtr<FChangelistFile
 static FString GetSharedBranchPath(const TMap<FString, FString>& InChangelistRecord)
 {
 	FString SharedBranchPath;
-	if (const FString* Found = InChangelistRecord.Find(FileDepotKey + "0"))
+	if (const FString* Found = InChangelistRecord.Find(ReviewHelpers::FileDepotKey + "0"))
 	{
 		SharedBranchPath = *Found;
 	}
@@ -505,7 +505,7 @@ static FString GetSharedBranchPath(const TMap<FString, FString>& InChangelistRec
 	//Each file in p4 has an index "depotFile0" this is the index that is updated by the while loop to find all the files "depotFile1", "depotFile2" etc...
 	uint32  RecordFileIndex = 1;
 	//The p4 records is the map a file key starts with "depotFile" and is followed by file index 
-	FString RecordFileMapKey = FileDepotKey + LexToString(RecordFileIndex);
+	FString RecordFileMapKey = ReviewHelpers::FileDepotKey + LexToString(RecordFileIndex);
 	while (const FString* Found = InChangelistRecord.Find(RecordFileMapKey))
 	{
 		// find starting from the left, find the portion that's shared between both strings
@@ -522,16 +522,16 @@ static FString GetSharedBranchPath(const TMap<FString, FString>& InChangelistRec
 		SharedBranchPath = SharedBranchPath.Left(TrimIndex);
 
 		// increment to next file path
-		RecordFileMapKey = FileDepotKey + LexToString(++RecordFileIndex);
+		RecordFileMapKey = ReviewHelpers::FileDepotKey + LexToString(++RecordFileIndex);
 	}
 	return SharedBranchPath;
 }
 
 void SSourceControlReview::SetChangelistInfo(const TMap<FString, FString>& InChangelistRecord)
 {
-	CurrentChangelistInfo.Author = FText::FromString(InChangelistRecord[AuthorKey]);
-	CurrentChangelistInfo.Description = FText::FromString(InChangelistRecord[DescriptionKey]);
-	CurrentChangelistInfo.Status = FText::FromString(InChangelistRecord[ChangelistStatusKey]);
+	CurrentChangelistInfo.Author = FText::FromString(InChangelistRecord[ReviewHelpers::AuthorKey]);
+	CurrentChangelistInfo.Description = FText::FromString(InChangelistRecord[ReviewHelpers::DescriptionKey]);
+	CurrentChangelistInfo.Status = FText::FromString(InChangelistRecord[ReviewHelpers::ChangelistStatusKey]);
 	CurrentChangelistInfo.SharedPath = FText::FromString(GetSharedBranchPath(InChangelistRecord));
 }
 
