[Feature] Review Tool Redesign

diff --git a/Engine/Plugins/Editor/ChangelistReview/Source/ChangelistReview.Build.cs b/Engine/Plugins/Editor/ChangelistReview/Source/ChangelistReview.Build.cs
--- a/Engine/Plugins/Editor/ChangelistReview/Source/ChangelistReview.Build.cs
+++ b/Engine/Plugins/Editor/ChangelistReview/Source/ChangelistReview.Build.cs
@@ -24,6 +24,7 @@ public class ChangelistReview : ModuleRules
 				"EditorStyle",
 				"ToolMenus",
 				"SourceControl",
+				"InputCore"
 			}
 		);
 	}
diff --git a/Engine/Plugins/Editor/ChangelistReview/Source/Private/ChangelistReviewModule.cpp b/Engine/Plugins/Editor/ChangelistReview/Source/Private/ChangelistReviewModule.cpp
--- a/Engine/Plugins/Editor/ChangelistReview/Source/Private/ChangelistReviewModule.cpp
+++ b/Engine/Plugins/Editor/ChangelistReview/Source/Private/ChangelistReviewModule.cpp
@@ -7,7 +7,7 @@
 #include "SSourceControlReview.h"
 #include "ToolMenus.h"
 
-#define LOCTEXT_NAMESPACE "FChangelistReviewModule"
+#define LOCTEXT_NAMESPACE "ChangelistReviewModule"
 
 static const FName SourceControlReviewTabName = FName(TEXT("SourceControlChangelistReview"));
 
@@ -24,7 +24,9 @@ void FChangelistReviewModule::StartupModule()
 		SourceControlIcon,
 		FUIAction(
 			FExecuteAction::CreateRaw(this, &FChangelistReviewModule::ShowReviewTab),
-			FCanExecuteAction::CreateRaw(this, &FChangelistReviewModule::CanShowReviewTab)
+			FCanExecuteAction(),
+			FGetActionCheckState(),
+			FIsActionButtonVisible::CreateRaw(this, &FChangelistReviewModule::CanShowReviewTab)
 		)
 	);
 	
diff --git a/Engine/Plugins/Editor/ChangelistReview/Source/Private/SChangelistEditableText.cpp b/Engine/Plugins/Editor/ChangelistReview/Source/Private/SChangelistEditableText.cpp
deleted file mode 100644
diff --git a/Engine/Plugins/Editor/ChangelistReview/Source/Private/SChangelistEditableText.h b/Engine/Plugins/Editor/ChangelistReview/Source/Private/SChangelistEditableText.h
deleted file mode 100644
diff --git a/Engine/Plugins/Editor/ChangelistReview/Source/Private/SSourceControlReview.cpp b/Engine/Plugins/Editor/ChangelistReview/Source/Private/SSourceControlReview.cpp
--- a/Engine/Plugins/Editor/ChangelistReview/Source/Private/SSourceControlReview.cpp
+++ b/Engine/Plugins/Editor/ChangelistReview/Source/Private/SSourceControlReview.cpp
@@ -2,9 +2,9 @@
 
 #include "SSourceControlReview.h"
 
+#include "ClassIconFinder.h"
 #include "ISourceControlModule.h"
 #include "ISourceControlProvider.h"
-#include "SChangelistEditableText.h"
 #include "SourceControlOperations.h"
 #include "SSourceControlReviewEntry.h"
 #include "AssetRegistry/AssetRegistryModule.h"
@@ -18,16 +18,17 @@
 #include "Widgets/Layout/SBox.h"
 #include "Misc/App.h"
 #include "Misc/MessageDialog.h"
+#include "Styling/StyleColors.h"
+#include "Widgets/Layout/SGridPanel.h"
+#include "AssetRegistry/AssetRegistryModule.h"
+#include "Brushes/SlateRoundedBoxBrush.h"
 
 #define LOCTEXT_NAMESPACE "SourceControlReview"
 
+using namespace SourceControlReview;
+
 namespace ReviewHelpers
 {
-	const FText EnterChangelistText(LOCTEXT("EnterChangelistText", "Enter a changelist number to diff:"));
-	const FText EnterChangelistTooltip(LOCTEXT("EnterChangelistTooltip", "Enter changelist"));
-	const FText LoadChangelistText(LOCTEXT("LoadChangelistText", "Load Changelist"));
-	const FText LoadingText(LOCTEXT("LoadingText", "Loading..."));
-	
 	const FString FileDepotKey = TEXT("depotFile");
 	const FString FileRevisionKey = TEXT("rev");
 	const FString FileActionKey = TEXT("action");
@@ -39,10 +40,43 @@ namespace ReviewHelpers
 	constexpr int32 RecordIndex = 0;
 }
 
+const UClass* FChangelistFileData::GetIconClass()
+{
+	// If we haven't cached the icon class yet, find it
+	if (!CachedIconClass.IsSet())
+	{
+		CachedIconClass = nullptr;
+		const FAssetRegistryModule& AssetRegistryModule = FModuleManager::LoadModuleChecked<FAssetRegistryModule>(TEXT("AssetRegistry"));
+		FString TempPackageName;
+		if (FPackageName::TryConvertFilenameToLongPackageName(ReviewFileName, TempPackageName))
+		{
+			TArray<FAssetData> OutAssetData;
+			AssetRegistryModule.Get().GetAssetsByPackageName(*TempPackageName, OutAssetData);
+
+			if (OutAssetData.Num() > 0)
+			{
+				CachedIconClass = FClassIconFinder::GetIconClassForAssetData(OutAssetData[0]);
+			}
+		}
+	}
+
+	return CachedIconClass.GetValue();
+}
+
 void SSourceControlReview::Construct(const FArguments& InArgs)
 {
 	const FString ProjectName = FApp::GetProjectName();
 	ensureAlwaysMsgf((!ProjectName.IsEmpty()), TEXT("BlueprintReviewTool - Unable to get ProjectName"));
+
+	const static FSlateRoundedBoxBrush RecessedBrush(FStyleColors::Recessed, CoreStyleConstants::InputFocusRadius);
+	const static FEditableTextBoxStyle InfoWidgetStyle =
+		FEditableTextBoxStyle(FAppStyle::Get().GetWidgetStyle<FEditableTextBoxStyle>("NormalEditableTextBox"))
+		.SetBackgroundImageNormal(RecessedBrush)
+		.SetBackgroundImageHovered(RecessedBrush)
+		.SetBackgroundImageFocused(RecessedBrush)
+		.SetBackgroundImageReadOnly(RecessedBrush);
+
+	const static FMargin InfoWidgetMargin(4.f, 2.f, 4.f, 8.f);
 	
 	ChildSlot
 	[
@@ -54,55 +88,170 @@ void SSourceControlReview::Construct(const FArguments& InArgs)
 			+SVerticalBox::Slot()
 			.AutoHeight()
 			[
-				SNew(SHorizontalBox)
-				+SHorizontalBox::Slot()
-				.VAlign(VAlign_Center)
-				.AutoWidth()
+				SAssignNew(ChangelistInfoWidget, SGridPanel)
+				.FillColumn(1, 1.f)
+
+				// Changelist
+				+SGridPanel::Slot(0,0)
+				.Padding(InfoWidgetMargin)
+				.HAlign(HAlign_Right)
 				[
 					SNew(STextBlock)
-					.Text(ReviewHelpers::EnterChangelistText)
-					.Font(FStyleFonts::Get().LargeBold)
+					.Text(LOCTEXT("ChangelistNumber", "Changelist"))
+					.Font(FStyleFonts::Get().Normal)
 				]
-				+SHorizontalBox::Slot()
-				.VAlign(VAlign_Center)
-				.Padding(5.f,0.f)
-				.AutoWidth()
+				+SGridPanel::Slot(1,0)
+				.Padding(4.f, 0.f, 4.f, 8.f)
+				.HAlign(HAlign_Left)
 				[
-					SNew(SBorder)
-					.ToolTipText(ReviewHelpers::EnterChangelistTooltip)
+					SNew(SHorizontalBox)
+					+SHorizontalBox::Slot()
+					.VAlign(VAlign_Center)
+					.Padding(0.f, 0.f, 8.f, 0.f)
+					.AutoWidth()
 					[
-						SAssignNew(ChangelistNumWidget, SChangelistEditableText)
-						.Font(FStyleFonts::Get().Large)
-						.MinDesiredWidth(200.f)
+						SAssignNew(ChangelistNumWidget, SEditableTextBox)
+						.Font(FStyleFonts::Get().Normal)
+						.MinDesiredWidth(55.f)
 						.Justification(ETextJustify::Center)
-						.Style(&FAppStyle::Get().GetWidgetStyle<FEditableTextStyle>("NormalEditableText"))
 						.OnTextCommitted(this, &SSourceControlReview::OnChangelistNumCommitted)
+						.OnTextChanged(this, &SSourceControlReview::OnChangelistNumChanged)
+						.BackgroundColor(FStyleColors::Input)
+					]
+					+SHorizontalBox::Slot()
+					.VAlign(VAlign_Center)
+					[
+						SNew(SButton)
+						.HAlign(HAlign_Center)
+						.OnClicked(this, &SSourceControlReview::OnLoadChangelistClicked)
+						[
+							SNew(STextBlock)
+							.Text(LOCTEXT("LoadChangelistText", "Load"))
+							.Font(FStyleFonts::Get().Normal)
+						]
 					]
 				]
-				+SHorizontalBox::Slot()
-				.VAlign(VAlign_Center)
+
+				// Author
+				+SGridPanel::Slot(0,1)
+				.Padding(InfoWidgetMargin)
+				.HAlign(HAlign_Right)
+				[
+					SNew(STextBlock)
+					.Text(LOCTEXT("ChangelistAuthor", "Author"))
+					.Font(FStyleFonts::Get().Normal)
+				]
+				+SGridPanel::Slot(1,1)
+				.Padding(4.f, 0.f, 4.f, 8.f)
+				.HAlign(HAlign_Fill)
+				[
+					SNew(SEditableTextBox)
+					.Text_Lambda([this]{return CurrentChangelistInfo.Author;})
+					.IsReadOnly(true)
+					.Style(&InfoWidgetStyle)
+				]
+
+				// Path
+				+SGridPanel::Slot(0,2)
+				.Padding(InfoWidgetMargin)
+				.HAlign(HAlign_Right)
+				[
+					SNew(STextBlock)
+					.Text(LOCTEXT("ChangelistPath", "Path"))
+					.Font(FStyleFonts::Get().Normal)
+				]
+				+SGridPanel::Slot(1,2)
+				.Padding(4.f, 0.f, 4.f, 8.f)
+				.HAlign(HAlign_Fill)
+				[
+					SNew(SEditableTextBox)
+					.Text_Lambda([this]{return CurrentChangelistInfo.SharedPath;})
+					.IsReadOnly(true)
+					.Style(&InfoWidgetStyle)
+				]
+
+				// Status
+				+SGridPanel::Slot(0,3)
+				.Padding(InfoWidgetMargin)
+				.HAlign(HAlign_Right)
 				[
-					SNew(SButton)
-					.HAlign(HAlign_Center)
-					.OnClicked(this, &SSourceControlReview::OnLoadChangelistClicked)
+					SNew(STextBlock)
+					.Text(LOCTEXT("ChangelistStatus", "Status"))
+					.Font(FStyleFonts::Get().Normal)
+				]
+				+SGridPanel::Slot(1,3)
+				.Padding(4.f, 0.f, 4.f, 8.f)
+				.HAlign(HAlign_Fill)
+				[
+					SNew(SEditableTextBox)
+					.Text_Lambda([this]{return CurrentChangelistInfo.Status;})
+					.IsReadOnly(true)
+					.Style(&InfoWidgetStyle)
+				]
+
+				// Description
+				+SGridPanel::Slot(0,4)
+				.Padding(InfoWidgetMargin)
+				.HAlign(HAlign_Right)
+				[
+					SNew(STextBlock)
+					.Text(LOCTEXT("ChangelistDescription", "Description"))
+					.Font(FStyleFonts::Get().Normal)
+				]
+				+SGridPanel::Slot(1,4)
+				.Padding(InfoWidgetMargin)
+				.HAlign(HAlign_Fill)
+				[
+					SNew(SBox)
+					.MaxDesiredHeight(147.f)
+					.MinDesiredHeight(147.f)
 					[
-						SNew(STextBlock)
-						.Text(ReviewHelpers::LoadChangelistText)
-						.Font(FStyleFonts::Get().LargeBold)
+						SNew(SMultiLineEditableTextBox)
+						.Text_Lambda([this]{return CurrentChangelistInfo.Description;})
+						.AutoWrapText(true)
+						.IsReadOnly(true)
+						.Style(&InfoWidgetStyle)
 					]
 				]
 			]
+			
+			+SVerticalBox::Slot()
+			.VAlign(VAlign_Top)
+			[
+				SAssignNew(ChangelistEntriesWidget, SListView<TSharedPtr<FChangelistFileData>>)
+				.ListItemsSource(&ChangelistFiles)
+				.OnGenerateRow(this, &SSourceControlReview::OnGenerateFileRow)
+				.HeaderRow(
+					SNew(SHeaderRow)
+					+HeaderColumn(ColumnIds::Status)
+					+HeaderColumn(ColumnIds::File)
+					+HeaderColumn(ColumnIds::Tools)
+				)
+			]
+			
+			+SVerticalBox::Slot()
+			.Padding(0.f, 0.f, 0.f, 98.f)
+			.HAlign(HAlign_Center)
+			.AutoHeight()
+			[
+				SAssignNew(EnterChangelistTextBlock, STextBlock)
+				.Visibility(EVisibility::Visible)
+				.ColorAndOpacity(FStyleColors::AccentGray)
+				.Text(LOCTEXT("EnterChangelistText", "Enter a Changelist number above to search"))
+			]
+			
 			+SVerticalBox::Slot()
 			.Padding(0.f, 10.f)
 			.HAlign(HAlign_Center)
-			.VAlign(VAlign_Center)
+			.VAlign(VAlign_Bottom)
 			.AutoHeight()
 			[
 				SAssignNew(LoadingTextBlock, STextBlock)
 				.Visibility(EVisibility::Collapsed)
-				.Text(ReviewHelpers::LoadingText)
+				.Text(LOCTEXT("LoadingText", "Loading..."))
 				.Font(FStyleFonts::Get().Large)
 			]
+			
 			+SVerticalBox::Slot()
 			.Padding(0.f, 15.f)
 			.AutoHeight()
@@ -111,115 +260,12 @@ void SSourceControlReview::Construct(const FArguments& InArgs)
 				.Visibility(EVisibility::Collapsed)
 				.Percent(1.f)
 			]
-			+SVerticalBox::Slot()
-			.AutoHeight()
-			.Padding(0.f, 20.f)
-			[
-				GetChangelistInfoWidget()
-			]
-			+SVerticalBox::Slot()
-			[
-				SNew(SScrollBox)
-				+SScrollBox::Slot()
-				[
-					SAssignNew(ChangelistEntriesWidget, SVerticalBox)
-				]
-			]
 		]
 	];
 }
 
-TSharedRef<SWidget> SSourceControlReview::GetChangelistInfoWidget()
-{
-	auto GenerateRow = [](FText Key, const TAttribute<FText>& Value)->TSharedRef<SWidget>
-	{
-		const TSharedRef<SWidget> KeyWidget = SNew(STextBlock)
-			.Text(FText::Format(LOCTEXT("ChangelistInfoKey", "{0}:"), Key))
-			.Font(FStyleFonts::Get().Large);
-
-		const TSharedRef<SWidget> ValueWidget = SNew(SBox)
-			.MaxDesiredHeight(125.f)
-			[
-				SNew(SScrollBox)
-				+SScrollBox::Slot()
-				[
-					SNew(SMultiLineEditableTextBox)
-					.Text(Value)
-					.IsReadOnly(true)
-					.Font(FStyleFonts::Get().Normal)
-					.BackgroundColor(FLinearColor(0,0,0,0))
-				]
-			];
-
-		return SNew(SWidgetSwitcher)
-		.WidgetIndex_Lambda(
-			[Value]{
-				const FString ResolvedValue = Value.Get().ToString();
-				int32 Index;
-				return ResolvedValue.FindChar('\n', Index);
-			}
-		)
-		+SWidgetSwitcher::Slot()
-		[
-			SNew(SHorizontalBox)
-			+SHorizontalBox::Slot()
-			.VAlign(VAlign_Center)
-			.AutoWidth()
-			[
-				KeyWidget
-			]
-			+SHorizontalBox::Slot()
-			[
-				ValueWidget
-			]
-		]
-		+SWidgetSwitcher::Slot()
-		[
-			SNew(SHorizontalBox)
-			+SHorizontalBox::Slot()
-			.AutoWidth()
-			[
-				KeyWidget
-			]
-			+SHorizontalBox::Slot()
-			[
-				ValueWidget
-			]
-		];
-	};
-	
-	return SAssignNew(ChangelistInfoWidget, SVerticalBox)
-	.Visibility(EVisibility::Collapsed)
-	+SVerticalBox::Slot()
-	.AutoHeight()
-	[
-		GenerateRow(LOCTEXT("ChangelistAuthor", "Author"), TAttribute<FText>::CreateLambda(
-			[this]{return CurrentChangelistInfo.Author;}))
-	]
-	+SVerticalBox::Slot()
-	.AutoHeight()
-	[
-		GenerateRow(LOCTEXT("ChangelistPath", "Path"), TAttribute<FText>::CreateLambda(
-			[this]{return CurrentChangelistInfo.SharedPath;}))
-	]
-	+SVerticalBox::Slot()
-	.AutoHeight()
-	[
-		GenerateRow(LOCTEXT("ChangelistStatus", "Status"), TAttribute<FText>::CreateLambda(
-			[this]{return CurrentChangelistInfo.Status;}))
-	]
-	+SVerticalBox::Slot()
-	.AutoHeight()
-	[
-		GenerateRow(LOCTEXT("ChangelistDescription", "Description"), TAttribute<FText>::CreateLambda(
-			[this]{return CurrentChangelistInfo.Description;}))
-	];
-}
-
 void SSourceControlReview::LoadChangelist(const FString& Changelist)
 {
-	ChangelistEntriesWidget->ClearChildren();
-	
 	if (!IsLoading())
 	{
 		SetLoading(true);
@@ -230,7 +276,7 @@ void SSourceControlReview::LoadChangelist(const FString& Changelist)
 		return;
 	}
 
-	ChangelistFileDataMap.Empty();
+	ChangelistFiles.Empty();
 
 	//This command runs p4 -describe (or similar for other version controls) to retrieve changelist record information
 	const TSharedRef<FGetChangelistDetails> GetChangelistDetailsCommand = ISourceControlOperation::Create<FGetChangelistDetails>();
@@ -350,6 +396,44 @@ FReply SSourceControlReview::OnLoadChangelistClicked()
 	return FReply::Handled();
 }
 
+void SSourceControlReview::OnChangelistNumChanged(const FText& Text)
+{
+	const FString& Data = Text.ToString();
+
+	// Use the longest substring that consists of only valid characters.
+	// For example if someone enters:
+	// "john.doe2/CL_123456789/version_13", we'll use "123456789" because it's longer than "2" and "13"
+	int32 BestSubstrBegin = 0;
+	int32 BestSubstrLen = 0;
+	int32 I = 0;
+	while (I < Data.Len())
+	{
+		// skip invalid characters
+		while(I < Data.Len() && !TChar<TCHAR>::IsDigit(Data[I]))
+		{
+			++I;
+		}
+
+		// grab chunk of consecutive valid characters
+		const int32 SubstrBegin = I;
+		while(I < Data.Len() && TChar<TCHAR>::IsDigit(Data[I]))
+		{
+			++I;
+		}
+		
+		// if found substr is longer than the best so far, update the best so far
+		if (I - SubstrBegin > BestSubstrLen)
+		{
+			BestSubstrLen = I - SubstrBegin;
+			BestSubstrBegin = SubstrBegin;
+		}
+	}
+	
+	const FString ValidData = Data.Mid(BestSubstrBegin, BestSubstrLen);
+	const FText ValidText = FText::FromString(ValidData);
+	ChangelistNumWidget->SetText(ValidText);
+}
+
 void SSourceControlReview::OnChangelistNumCommitted(const FText& Text, ETextCommit::Type CommitMethod)
 {
 	if (CommitMethod == ETextCommit::OnEnter)
@@ -362,7 +446,7 @@ void SSourceControlReview::OnGetFileFromSourceControl(TSharedPtr<FChangelistFile
 {
 	if (ChangelistFileData && ChangelistFileData->IsDataValidForEntry())
 	{
-		ChangelistFileDataMap.Add(ChangelistFileData->ReviewFileName, ChangelistFileData);
+		ChangelistFiles.Add(ChangelistFileData);
 	}
 
 	FilesLoaded++;
@@ -373,15 +457,13 @@ void SSourceControlReview::OnGetFileFromSourceControl(TSharedPtr<FChangelistFile
 	{
 		const FAssetRegistryModule& AssetRegistryModule = FModuleManager::LoadModuleChecked<FAssetRegistryModule>( TEXT( "AssetRegistry" ) );
 
-		TArray<TSharedPtr<FChangelistFileData>> ChangelistFilesData;
-		ChangelistFileDataMap.GenerateValueArray(ChangelistFilesData);
-		Algo::Sort(ChangelistFilesData, [](const TSharedPtr<FChangelistFileData> &A, const TSharedPtr<FChangelistFileData> &B)
+		Algo::Sort(ChangelistFiles, [](const TSharedPtr<FChangelistFileData> &A, const TSharedPtr<FChangelistFileData> &B)
 		{
 			return A->RelativeFilePath < B->RelativeFilePath;
 		});
 		
 		TArray<FString> ChangelistFilePaths;
-		for (const TSharedPtr<FChangelistFileData>& FileData : ChangelistFilesData)
+		for (const TSharedPtr<FChangelistFileData>& FileData : ChangelistFiles)
 		{
 			if(FileData)
 			{
@@ -390,11 +472,17 @@ void SSourceControlReview::OnGetFileFromSourceControl(TSharedPtr<FChangelistFile
 		}
 		
 		AssetRegistryModule.Get().ScanFilesSynchronous(ChangelistFilePaths);
-		
-		for(const TSharedPtr<FChangelistFileData>& FileData : ChangelistFilesData)
+
+		// now that the files are in the asset registry, cache their associated class so their class icons can be created quickly
+		for (const TSharedPtr<FChangelistFileData>& FileData : ChangelistFiles)
 		{
-			AddDiffFile(*FileData);
+			if(FileData)
+			{
+				FileData->GetIconClass();
+			}
 		}
+		
+		ChangelistEntriesWidget->RebuildList();
 		SetLoading(false);
 	}
 }
@@ -406,9 +494,14 @@ void SSourceControlReview::SetLoading(bool bInLoading)
 	// show loading bar and text if we're loading
 	LoadingProgressBar->SetVisibility(bInLoading? EVisibility::Visible : EVisibility::Collapsed);
 	LoadingTextBlock->SetVisibility(bInLoading? EVisibility::Visible : EVisibility::Collapsed);
-
-	// hide CL info if we're loading or there's no changelist
-	ChangelistInfoWidget->SetVisibility((ChangelistNumWidget->GetText().IsEmpty() || bInLoading)? EVisibility::Collapsed : EVisibility::Visible);
+	if (bInLoading)
+	{
+		EnterChangelistTextBlock->SetVisibility(EVisibility::Collapsed);
+	}
+	else if(ChangelistFiles.IsEmpty())
+	{
+		EnterChangelistTextBlock->SetVisibility(EVisibility::Visible);
+	}
 }
 
 bool SSourceControlReview::IsLoading() const
@@ -535,14 +628,10 @@ void SSourceControlReview::SetChangelistInfo(const TMap<FString, FString>& InCha
 	CurrentChangelistInfo.SharedPath = FText::FromString(GetSharedBranchPath(InChangelistRecord));
 }
 
-void SSourceControlReview::AddDiffFile(const FChangelistFileData& FileData) const
+TSharedRef<ITableRow> SSourceControlReview::OnGenerateFileRow(TSharedPtr<FChangelistFileData> FileData, const TSharedRef<STableViewBase>& Table) const
 {
-	ChangelistEntriesWidget->AddSlot()
-	.Padding(0.f, 5.f)
-	[
-		SNew(SSourceControlReviewEntry)
-		.FileData(FileData)
-	];
+	return SNew(SSourceControlReviewEntry, Table)
+	.FileData(*FileData);
 }
 
 FString SSourceControlReview::TrimSharedPath(FString FullCLPath) const
@@ -551,6 +640,39 @@ FString SSourceControlReview::TrimSharedPath(FString FullCLPath) const
 	return FullCLPath;
 }
 
+SHeaderRow::FColumn::FArguments SSourceControlReview::HeaderColumn(FName HeaderName)
+{
+	FText ColumnLabel;
+	TOptional<float> ColumnWidth;
+	if (HeaderName == ColumnIds::Status)
+	{
+		ColumnLabel = LOCTEXT("StatusColumnHeader", "Status");
+		ColumnWidth = 60;
+	}
+	else if (HeaderName == ColumnIds::File)
+	{
+		ColumnLabel = LOCTEXT("FileColumnHeader", "File");
+	}
+	else if (HeaderName == ColumnIds::Tools)
+	{
+		ColumnLabel = LOCTEXT("ToolsColumnHeader", "Tools");
+		ColumnWidth = 88;
+	}
+
+	return SHeaderRow::Column(HeaderName)
+		.FixedWidth(ColumnWidth)
+		.HAlignHeader(HeaderName == ColumnIds::File? HAlign_Fill : HAlign_Center)
+		.HAlignCell(HAlign_Fill)
+		.VAlignCell(VAlign_Fill)
+		.VAlignHeader(VAlign_Fill)
+		.HeaderContentPadding(FMargin(10.f, 6.f))
+		[
+			SNew(STextBlock)
+			.Text(ColumnLabel)
+			.Font(FStyleFonts::Get().Normal)
+		];
+}
+
 FString SSourceControlReview::AsAssetPath(const FString& FullCLPath)
 {
 	const FString ProjectName = FApp::GetProjectName();
diff --git a/Engine/Plugins/Editor/ChangelistReview/Source/Private/SSourceControlReview.h b/Engine/Plugins/Editor/ChangelistReview/Source/Private/SSourceControlReview.h
--- a/Engine/Plugins/Editor/ChangelistReview/Source/Private/SSourceControlReview.h
+++ b/Engine/Plugins/Editor/ChangelistReview/Source/Private/SSourceControlReview.h
@@ -6,8 +6,10 @@
 #include "Widgets/SCompoundWidget.h"
 #include "Widgets/Text/STextBlock.h"
 #include "Widgets/SBoxPanel.h"
+#include "Widgets/Views/SListView.h"
 
-class SChangelistEditableText;
+class SSourceControlReviewEntry;
+class SChangelistEditableTextBox;
 class SProgressBar;
 
 namespace SourceControlReview
@@ -75,7 +77,19 @@ namespace SourceControlReview
 		EChangelistState ChangelistState = EChangelistState::Submitted;
 
 		ESourceControlAction FileSourceControlAction = ESourceControlAction::Unset;
+
+		const UClass* GetIconClass();
+		
+	private:
+		TOptional<const UClass*> CachedIconClass;
 	};
+
+	namespace ColumnIds
+	{
+		inline const FName Status = TEXT("Status");
+		inline const FName File = TEXT("File");
+		inline const FName Tools = TEXT("Tools");
+	}
 }
 
 
@@ -104,6 +118,7 @@ public:
 private:
 	void OnChangelistLoadComplete(const FSourceControlOperationRef& InOperation, ECommandResult::Type InResult, FString Changelist);
 	FReply OnLoadChangelistClicked();
+	void OnChangelistNumChanged(const FText& Text);
 	void OnChangelistNumCommitted(const FText& Text, ETextCommit::Type CommitMethod);
 	
 	/**
@@ -142,17 +157,15 @@ private:
 	 */
 	void SetChangelistInfo(const TMap<FString, FString>& InChangelistRecord);
 
-	/**
-	 * Adds a SSourceControlReviewEntry to the GUI
-	 * @param FileData The File being added to the list
-	 */
-	void AddDiffFile(const FChangelistFileData& FileData) const;
+	TSharedRef<ITableRow> OnGenerateFileRow(TSharedPtr<FChangelistFileData> FileData, const TSharedRef<STableViewBase>& Table) const;
 
 	/**
 	 * Removes CurrentChangelistInfo.SharedPath from the beginning of FullCLPath
 	 */
 	FString TrimSharedPath(FString FullCLPath) const;
 
+	static SHeaderRow::FColumn::FArguments HeaderColumn(FName HeaderName);
+
 	/**
 	 * Removes the game directory from the beginning of the FullCLPath
 	 */
@@ -162,16 +175,14 @@ private:
 	bool bChangelistLoading = false;
 	uint32 FilesToLoad = 0;
 	uint32 FilesLoaded = 0;
-	TMap<FString, TSharedPtr<FChangelistFileData>> ChangelistFileDataMap;
-
-	// builds a widget that displays a changelist's Author, Path, Status, and Description
-	TSharedRef<SWidget> GetChangelistInfoWidget();
+	TArray<TSharedPtr<FChangelistFileData>> ChangelistFiles;
 	
-	TSharedPtr<SChangelistEditableText> ChangelistNumWidget;
+	TSharedPtr<SEditableTextBox> ChangelistNumWidget;
+	TSharedPtr<STextBlock> EnterChangelistTextBlock;
 	TSharedPtr<STextBlock> LoadingTextBlock;
 	TSharedPtr<SProgressBar> LoadingProgressBar;
 	TSharedPtr<SWidget> ChangelistInfoWidget;
-	TSharedPtr<SVerticalBox> ChangelistEntriesWidget;
+	TSharedPtr<SListView<TSharedPtr<FChangelistFileData>>> ChangelistEntriesWidget;
 
 	// Info about the current chagnelist
 	struct FChangelistInfo
diff --git a/Engine/Plugins/Editor/ChangelistReview/Source/Private/SSourceControlReviewEntry.cpp b/Engine/Plugins/Editor/ChangelistReview/Source/Private/SSourceControlReviewEntry.cpp
--- a/Engine/Plugins/Editor/ChangelistReview/Source/Private/SSourceControlReviewEntry.cpp
+++ b/Engine/Plugins/Editor/ChangelistReview/Source/Private/SSourceControlReviewEntry.cpp
@@ -12,9 +12,11 @@
 #include "Misc/MessageDialog.h"
 #include "Settings/EditorLoadingSavingSettings.h"
 #include "Styling/StarshipCoreStyle.h"
+#include "Styling/StyleColors.h"
 #include "Widgets/Text/STextBlock.h"
 #include "Widgets/Layout/SBox.h"
 #include "Widgets/Input/SButton.h"
+#include "Widgets/Layout/SScrollBox.h"
 
 #define LOCTEXT_NAMESPACE "SourceControlReviewEntry"
 
@@ -23,129 +25,135 @@ namespace ReviewEntryConsts
 	static const FString TempFolder = TEXT("/Temp/");
 }
 
-void SSourceControlReviewEntry::Construct(const FArguments& InArgs)
+void SSourceControlReviewEntry::Construct(const FArguments& InArgs, const TSharedRef<STableViewBase>& InOwnerTableView)
 {
-	static const UFont* Font = LoadObject<UFont>(nullptr, TEXT("/Game/UI/Foundation/Fonts/NotoSans.NotoSans"));
-	
 	ChangelistFileData = InArgs._FileData;
-	ChildSlot
-	[
-		SNew(SHorizontalBox)
-		
-		+SHorizontalBox::Slot()
-		.AutoWidth()
-		.Padding(10.f, 0.f)
-		.HAlign(HAlign_Center)
-		.VAlign(VAlign_Center)
-		[
-			SAssignNew(SourceActionIcon, SImage)
-			.Image(this, &SSourceControlReviewEntry::GetSourceControlIcon)
-		]
-		
-		+SHorizontalBox::Slot()
-		.FillWidth(1.f)
-		.Padding(0.f, 0.f, 5.f, 0.f)
-		.VAlign(VAlign_Center)
-		[
-			SNew(SVerticalBox)
-			+SVerticalBox::Slot()
-			.AutoHeight()
+	SMultiColumnTableRow::Construct(FSuperRowType::FArguments(), InOwnerTableView);
+
+	// figure out how this asset diffs, and bind it to this->DiffMethod
+	TryBindDiffMethod();
+}
+
+TSharedRef<SWidget> SSourceControlReviewEntry::GenerateWidgetForColumn(const FName& ColumnName)
+{
+	TSharedPtr<SWidget> InnerContent;
+	
+	if (ColumnName == SourceControlReview::ColumnIds::Status)
+	{
+		SAssignNew(InnerContent, SImage)
+		.DesiredSizeOverride(FVector2D(20.f, 20.f))
+		.ColorAndOpacity(this, &SSourceControlReviewEntry::GetSourceControlIconColor)
+		.Image(this, &SSourceControlReviewEntry::GetSourceControlIcon);
+	}
+	else if (ColumnName == SourceControlReview::ColumnIds::File)
+	{
+		SAssignNew(InnerContent, SVerticalBox)
+		+SVerticalBox::Slot()
+		.AutoHeight()
+		.Padding(16.f, 1.f, 0.f, 0.f)
+		[	
+			SNew(SHorizontalBox)
+			+SHorizontalBox::Slot()
+			.HAlign(HAlign_Left)
+			.AutoWidth()
+			.Padding(0.f, 0.f, 8.f, 4.f)
+			[
+				SNew(SBox)
+				.MaxAspectRatio(1.f)
+				.MinAspectRatio(1.f)
+				.HAlign(HAlign_Left)
+				[
+							
+					SAssignNew(AssetTypeIcon, SImage)
+					.DesiredSizeOverride(FVector2D(16.f, 16.f))
+					.Image(GetAssetTypeIcon())
+					.ToolTipText(GetAssetType())
+				]
+			]
+			+SHorizontalBox::Slot()
+			.FillWidth(1.f)
 			[
 				SNew(STextBlock)
 				.Text(this, &SSourceControlReviewEntry::GetAssetNameText)
-				.Font(FStyleFonts::Get().Large)
+				.ColorAndOpacity(FStyleColors::AccentWhite)
+				.Font(FStyleFonts::Get().Normal)
 			]
-			+SVerticalBox::Slot()
-			.AutoHeight()
+		]
+		+SVerticalBox::Slot()
+		.AutoHeight()
+		.Padding(16.f, 0.f, 0.f, 5.f)
+		[
+			SNew(SScrollBox)
+			.Orientation(EOrientation::Orient_Horizontal)
+			+SScrollBox::Slot()
 			[
 				SNew(STextBlock)
 				.Text(this, &SSourceControlReviewEntry::GetLocalAssetPathText)
 				.Font(FStyleFonts::Get().Small)
 			]
-		]
-		
+		];
+	}
+	else if (ColumnName == SourceControlReview::ColumnIds::Tools)
+	{
+		SAssignNew(InnerContent, SHorizontalBox)
 		+SHorizontalBox::Slot()
-		.AutoWidth()
-		.Padding(10.f, 0.f)
-		.HAlign(HAlign_Center)
-		.VAlign(VAlign_Center)
+		.Padding(0.f, 0.f, 7.f, 0.f)
 		[
-			SNew(SBox)
-			.WidthOverride(32.f)
-			.HeightOverride(32.f)
+			SNew(SButton)
+			.ButtonStyle(FAppStyle::Get(), "SimpleButton")
+			.OnClicked(this, &SSourceControlReviewEntry::OnDiffClicked)
+	 		.ToolTipText(LOCTEXT("ViewDiffTooltip", "Diff Against Previous Revision"))
+			.ContentPadding(FMargin(0.f))
 			[
-				SAssignNew(AssetTypeIcon, SImage)
-				.Image(GetAssetTypeIcon())
+				SNew(SImage)
+				.ColorAndOpacity(FSlateColor::UseForeground())
+				.DesiredSizeOverride(FVector2D(20.f, 20.f))
+				.Image(FAppStyle::Get().GetBrush("SourceControl.Actions.Diff"))
 			]
 		]
-		
-		+SHorizontalBox::Slot()
-		.AutoWidth()
-		.HAlign(HAlign_Center)
-		.VAlign(VAlign_Center)
-		[
-			SAssignNew(FileDeletedTextBlock, STextBlock)
-		]
-		
 		+SHorizontalBox::Slot()
-		.FillWidth(1.f)
-		.Padding(10.f, 0.f)
-		.VAlign(VAlign_Center)
+		.Padding(7.f, 0.f, 0.f, 0.f)
 		[
-			SAssignNew(ReviewInputsBox, SHorizontalBox)
-			+SHorizontalBox::Slot()
-			.FillWidth(1.f)
-			.Padding(0.f, 0.f, 5.f, 0.f)
-			[
-				SNew(SButton)
-				.VAlign(VAlign_Center)
-				[
-					SAssignNew(ViewDiffButtonText, STextBlock)
-					.Justification(ETextJustify::Center)
-					.Text(LOCTEXT("ViewDiffButton", "View Diff"))
-				]
-				.OnClicked(this, &SSourceControlReviewEntry::OnDiffClicked)
-			]
-			+SHorizontalBox::Slot()
-			.FillWidth(1.f)
-			.Padding(0.f, 0.f, 5.f, 0.f)
+			SNew(SButton)
+			.ButtonStyle(FAppStyle::Get(), "SimpleButton")
+			.IsEnabled(this, &SSourceControlReviewEntry::CanBrowseToAsset)
+			.OnClicked(this, &SSourceControlReviewEntry::OnBrowseToAssetClicked)
+			.ToolTipText(this, &SSourceControlReviewEntry::GetBrowseToAssetTooltip)
+			.ContentPadding(FMargin(0.f))
 			[
-				SNew(SButton)
-				.IsEnabled(this, &SSourceControlReviewEntry::CanBrowseToAsset)
-				.ContentPadding(FMargin(8.f, 4.f))
-				.OnClicked(this, &SSourceControlReviewEntry::OnBrowseToAssetClicked)
-				[
-					SNew(SHorizontalBox)
-					+SHorizontalBox::Slot()
-					.AutoWidth()
-					.VAlign(VAlign_Center)
-					.HAlign(HAlign_Left)
-					[
-						SNew(SImage)
-						.Image(FAppStyle::GetBrush("Icons.Search"))
-					]
-					+SHorizontalBox::Slot()
-					.FillWidth(1.f)
-					.VAlign(VAlign_Center)
-					.HAlign(HAlign_Center)
-					[
-						SNew(STextBlock)
-						.Justification(ETextJustify::Center)
-						.Text(LOCTEXT("BrowseToAssetButton", "Browse To Asset"))
-					]
-				]
+				SNew(SImage)
+				.ColorAndOpacity(FSlateColor::UseForeground())
+				.DesiredSizeOverride(FVector2D(20.f, 20.f))
+				.Image(FAppStyle::Get().GetBrush("Icons.BrowseContent"))
 			]
+			
+		];
+	}
+
+	static const FSlateBrush* const WhiteBrush = FAppStyle::GetBrush("WhiteBrush");
+	
+	return SNew(SBorder)
+	.BorderImage(WhiteBrush)
+	.BorderBackgroundColor(FStyleColors::AccentBlack)
+	.Padding(1.f)
+	[
+		SNew(SBorder)
+		.BorderImage(WhiteBrush)
+		.BorderBackgroundColor(FStyleColors::Recessed)
+		.Padding(0.f)
+		.VAlign(VAlign_Center)
+		.HAlign(ColumnName == SourceControlReview::ColumnIds::File? HAlign_Fill : HAlign_Center)
+		[
+			InnerContent.ToSharedRef()
 		]
 	];
-
-	// figure out how this asset diffs, and bind it to this->DiffMethod
-	TryBindDiffMethod();
 }
 
 void SSourceControlReviewEntry::SetEntryData(const FChangelistFileData& InChangelistFileData)
 {
 	ChangelistFileData = InChangelistFileData;
 	AssetTypeIcon->SetImage(GetAssetTypeIcon());
+	AssetTypeIcon->SetToolTipText(GetAssetType());
 
 	// if asset changed, we might diff differently. rebind the diff method.
 	DiffMethod.Unbind();
@@ -167,6 +175,15 @@ FReply SSourceControlReviewEntry::OnBrowseToAssetClicked() const
 	return FReply::Handled();
 }
 
+FText SSourceControlReviewEntry::GetBrowseToAssetTooltip() const
+{
+	if (CanBrowseToAsset())
+	{
+		return LOCTEXT("BrowseToAssetTooltip", "Browse To Asset");
+	}
+	return  LOCTEXT("CantBrowseToAssetTooltip", "This File is not an Asset");
+}
+
 bool SSourceControlReviewEntry::CanBrowseToAsset() const
 {
 	TArray<FAssetData> Assets;
@@ -312,26 +329,47 @@ const FSlateBrush* SSourceControlReviewEntry::GetSourceControlIcon() const
 	return Brushes[Index];
 }
 
-const FSlateBrush* SSourceControlReviewEntry::GetAssetTypeIcon() const
+FSlateColor SSourceControlReviewEntry::GetSourceControlIconColor() const
 {
-	const FAssetRegistryModule& AssetRegistryModule = FModuleManager::LoadModuleChecked<FAssetRegistryModule>( TEXT( "AssetRegistry" ) );
+	switch(ChangelistFileData.FileSourceControlAction)
+	{
+	case ESourceControlAction::Add:
+		return FAppStyle::Get().GetSlateColor("SourceControl.Diff.AdditionColor");
+	case ESourceControlAction::Edit:
+		return FAppStyle::Get().GetSlateColor("SourceControl.Diff.MajorModificationColor");
+	case ESourceControlAction::Delete:
+		return FAppStyle::Get().GetSlateColor("SourceControl.Diff.SubtractionColor");
+	
+	case ESourceControlAction::Branch:
+	case ESourceControlAction::Integrate:
+		return FAppStyle::Get().GetSlateColor("SourceControl.Diff.MinorModificationColor");
 	
-	FString TempPackageName;
-	if (FPackageName::TryConvertFilenameToLongPackageName(ChangelistFileData.ReviewFileName, TempPackageName))
+	default: return FStyleColors::Foreground.GetSpecifiedColor();
+	}
+}
+
+const FSlateBrush* SSourceControlReviewEntry::GetAssetTypeIcon()
+{
+	if (const UClass* IconClass = ChangelistFileData.GetIconClass())
 	{
-		TArray<FAssetData> OutAssetData;
-		AssetRegistryModule.Get().GetAssetsByPackageName( *TempPackageName, OutAssetData );
-		
-		if (OutAssetData.Num() > 0)
+		if (const FSlateBrush* FileTypeBrush = FClassIconFinder::FindThumbnailForClass(IconClass))
 		{
-			if (const FSlateBrush* FileTypeBrush = FClassIconFinder::FindThumbnailForClass(FClassIconFinder::GetIconClassForAssetData(OutAssetData[0])))
-			{
-				return FileTypeBrush;
-			}
+			return FileTypeBrush;
 		}
 	}
+	return FAppStyle::GetBrush("ContentBrowser.ColumnViewAssetIcon");
+}
+
+FText SSourceControlReviewEntry::GetAssetType()
+{
+	if (const UClass* IconClass = ChangelistFileData.GetIconClass())
+	{
+		return FText::FromString(IconClass->GetName());
+	}
 	
-	return FAppStyle::Get().GetBrush(TEXT("NoBrush"));
+	int32 ChopIndex;
+	ChangelistFileData.AssetFilePath.FindLastChar('.', ChopIndex);
+	return FText::Format(LOCTEXT("NonAssetFileType", "{0} File"), FText::FromString(ChangelistFileData.AssetFilePath.RightChop(ChopIndex)));
 }
 
 FText SSourceControlReviewEntry::GetAssetNameText() const
@@ -341,7 +379,7 @@ FText SSourceControlReviewEntry::GetAssetNameText() const
 
 FText SSourceControlReviewEntry::GetLocalAssetPathText() const
 {
-	return FText::FromString(ChangelistFileData.RelativeFilePath);
+	return FText::FromString(TEXT("...") + ChangelistFileData.RelativeFilePath);
 }
 
 const FString& SSourceControlReviewEntry::GetSearchableString() const
diff --git a/Engine/Plugins/Editor/ChangelistReview/Source/Private/SSourceControlReviewEntry.h b/Engine/Plugins/Editor/ChangelistReview/Source/Private/SSourceControlReviewEntry.h
--- a/Engine/Plugins/Editor/ChangelistReview/Source/Private/SSourceControlReviewEntry.h
+++ b/Engine/Plugins/Editor/ChangelistReview/Source/Private/SSourceControlReviewEntry.h
@@ -4,11 +4,11 @@
 #include "IAssetTypeActions.h"
 #include "SSourceControlReview.h"
 #include "Engine/Blueprint.h"
-#include "Widgets/Images/SImage.h"
 #include "Widgets/DeclarativeSyntaxSupport.h"
-#include "Widgets/SCompoundWidget.h"
 
-class SSourceControlReviewEntry : public SCompoundWidget
+class SImage;
+
+class SSourceControlReviewEntry : public SMultiColumnTableRow<TSharedPtr<SourceControlReview::FChangelistFileData>>
 {
 public:
 	using ESourceControlAction = SourceControlReview::ESourceControlAction;
@@ -20,13 +20,16 @@ public:
 	SLATE_END_ARGS()
 	
 	/** Constructs the widget */
-	void Construct(const FArguments& InArgs);
+	void Construct(const FArguments& InArgs, const TSharedRef<STableViewBase>& InOwnerTableView);
+	
+	virtual TSharedRef<SWidget> GenerateWidgetForColumn(const FName& ColumnName) override;
 	
 	/**
 	 * Sets needed data to show diff and browse to asset
 	 * @param InChangelistFileData changelist entry information
 	 */
 	void SetEntryData(const FChangelistFileData& InChangelistFileData);
+private:
 
 	/**
 	 * Hooks into built in engine diff functionality to show differences between supplied file revisions
@@ -38,7 +41,11 @@ public:
 	 */
 	FReply OnBrowseToAssetClicked() const;
 
-private:
+	/**
+	 * Uses engine sync to browse to asset functionality to navigate the content window to file
+	 */
+	FText GetBrowseToAssetTooltip() const;
+
 	/**
 	 * Checks if we have a valid local file for this entry to sync content window to
 	 */
@@ -78,11 +85,21 @@ private:
 	 * return an icon reflecting this entry's source control action (i.e. add, edit, delete, etc...)
 	 */
 	const FSlateBrush* GetSourceControlIcon() const;
+
+	/**
+	 * return the color associated with this entry's source control action (i.e. add, edit, delete, etc...)
+	 */
+	FSlateColor GetSourceControlIconColor() const;
+	
+	/**
+	 * return an icon that indicates the type of the asset being displayed
+	 */
+	const FSlateBrush* GetAssetTypeIcon();
 	
 	/**
 	 * return an icon that indicates the type of the asset being displayed
 	 */
-	const FSlateBrush* GetAssetTypeIcon() const;
+	FText GetAssetType();
 
 	/**
 	 * Name of the asset being displayed
@@ -113,10 +130,5 @@ private:
 	FDiffMethod DiffMethod;
 
 	// Sub-Widgets
-	TSharedPtr<SImage> SourceActionIcon;
 	TSharedPtr<SImage> AssetTypeIcon;
-	TSharedPtr<STextBlock> FileDeletedTextBlock;
-	TSharedPtr<SHorizontalBox> ReviewInputsBox;
-	TSharedPtr<SImage> ViewedEntryIcon;
-	TSharedPtr<STextBlock> ViewDiffButtonText;
 };
diff --git a/Engine/Source/Editor/EditorStyle/Private/StarshipStyle.cpp b/Engine/Source/Editor/EditorStyle/Private/StarshipStyle.cpp
--- a/Engine/Source/Editor/EditorStyle/Private/StarshipStyle.cpp
+++ b/Engine/Source/Editor/EditorStyle/Private/StarshipStyle.cpp
@@ -7091,6 +7091,12 @@ void FStarshipEditorStyle::FStyle::SetupSourceControlStyles()
 		Set("SourceControl.Actions.CheckOut", new CORE_IMAGE_BRUSH_SVG("Starship/SourceControl/SCC_CheckedOut", Icon16x16));
 		Set("SourceControl.Actions.Refresh", new CORE_IMAGE_BRUSH_SVG("Starship/Common/Undo", Icon16x16));
 		Set("SourceControl.Actions.NewChangelist", new CORE_IMAGE_BRUSH_SVG("Starship/Common/plus-circle", Icon16x16));
+
+		// Diff colors
+		Set("SourceControl.Diff.AdditionColor", FLinearColor(0.02f, 0.94f, 0.f));
+		Set("SourceControl.Diff.SubtractionColor", FLinearColor(1.f, 0.16f, 0.16f));
+		Set("SourceControl.Diff.MajorModificationColor", FLinearColor(0.04f, 0.87f, 1.f));
+		Set("SourceControl.Diff.MinorModificationColor", FLinearColor(0.74f, 0.69f, 0.79f));
 	}
 #endif // WITH_EDITOR || (IS_PROGRAM && WITH_UNREAL_DEVELOPER_TOOLS)
 
diff --git a/Engine/Source/Runtime/Engine/Public/DiffResults.h b/Engine/Source/Runtime/Engine/Public/DiffResults.h
--- a/Engine/Source/Runtime/Engine/Public/DiffResults.h
+++ b/Engine/Source/Runtime/Engine/Public/DiffResults.h
@@ -3,6 +3,7 @@
 #pragma once
 
 #include "CoreMinimal.h"
+#include "Styling/StyleColors.h"
 
 /** Differences found within a graph or object */
 namespace EDiffType
@@ -114,12 +115,12 @@ struct FDiffSingleResult
 	{
 		switch(Category)
 		{
-		case EDiffType::ADDITION: return FLinearColor(0.02f, 0.94f, 0.f);
-		case EDiffType::SUBTRACTION: return FLinearColor(1.f, 0.16f, 0.16f);
-		case EDiffType::MODIFICATION: return FLinearColor(0.04f, 0.87f, 1.f);
-		case EDiffType::MINOR:  return FLinearColor(0.74f, 0.69f, 0.79f);
+		case EDiffType::ADDITION: return FAppStyle::Get().GetSlateColor("SourceControl.Diff.AdditionColor").GetSpecifiedColor();
+		case EDiffType::SUBTRACTION: return FAppStyle::Get().GetSlateColor("SourceControl.Diff.SubtractionColor").GetSpecifiedColor();
+		case EDiffType::MODIFICATION: return FAppStyle::Get().GetSlateColor("SourceControl.Diff.MajorModificationColor").GetSpecifiedColor();
+		case EDiffType::MINOR:  return FAppStyle::Get().GetSlateColor("SourceControl.Diff.MinorModificationColor").GetSpecifiedColor();
 		
-		default: return FLinearColor::White;
+		default: return FStyleColors::Foreground.GetSpecifiedColor();
 		}
 	}
 
