[BugFix] CIS UnrealEditor NonUnity Win64 Error - SSourceControlReview.h(180): error C2065: 'SEditableTextBox': undeclared identifier

diff --git a/Engine/Plugins/Editor/ChangelistReview/Source/Private/SSourceControlReview.h b/Engine/Plugins/Editor/ChangelistReview/Source/Private/SSourceControlReview.h
--- a/Engine/Plugins/Editor/ChangelistReview/Source/Private/SSourceControlReview.h
+++ b/Engine/Plugins/Editor/ChangelistReview/Source/Private/SSourceControlReview.h
@@ -7,6 +7,7 @@
 #include "Widgets/Text/STextBlock.h"
 #include "Widgets/SBoxPanel.h"
 #include "Widgets/Views/SListView.h"
+#include "Widgets/Input/SEditableTextBox.h"
 
 class SSourceControlReviewEntry;
 class SChangelistEditableTextBox;
