Simplified SSourceControlReview::OnChangelistNumChanged by using regex

diff --git a/Engine/Plugins/Editor/ChangelistReview/Source/Private/SSourceControlReview.cpp b/Engine/Plugins/Editor/ChangelistReview/Source/Private/SSourceControlReview.cpp
--- a/Engine/Plugins/Editor/ChangelistReview/Source/Private/SSourceControlReview.cpp
+++ b/Engine/Plugins/Editor/ChangelistReview/Source/Private/SSourceControlReview.cpp
@@ -22,6 +22,7 @@
 #include "Widgets/Layout/SGridPanel.h"
 #include "AssetRegistry/AssetRegistryModule.h"
 #include "Brushes/SlateRoundedBoxBrush.h"
+#include "Internationalization/Regex.h"
 
 #define LOCTEXT_NAMESPACE "SourceControlReview"
 
@@ -403,34 +404,23 @@ void SSourceControlReview::OnChangelistNumChanged(const FText& Text)
 	// Use the longest substring that consists of only valid characters.
 	// For example if someone enters:
 	// "john.doe2/CL_123456789/version_13", we'll use "123456789" because it's longer than "2" and "13"
-	int32 BestSubstrBegin = 0;
-	int32 BestSubstrLen = 0;
-	int32 I = 0;
-	while (I < Data.Len())
+	int32 BestMatchBegin = 0;
+	int32 BestMatchLen = 0;
+	
+	FRegexMatcher RegexMatcher(FRegexPattern(TEXT("\\d+")), Data);
+	while(RegexMatcher.FindNext())
 	{
-		// skip invalid characters
-		while(I < Data.Len() && !TChar<TCHAR>::IsDigit(Data[I]))
-		{
-			++I;
-		}
-
-		// grab chunk of consecutive valid characters
-		const int32 SubstrBegin = I;
-		while(I < Data.Len() && TChar<TCHAR>::IsDigit(Data[I]))
-		{
-			++I;
-		}
+		const int32 MatchBegin = RegexMatcher.GetMatchBeginning();
+		const int32 MatchLen = RegexMatcher.GetMatchEnding() - MatchBegin;
 		
-		// if found substr is longer than the best so far, update the best so far
-		if (I - SubstrBegin > BestSubstrLen)
+		if (MatchLen > BestMatchLen)
 		{
-			BestSubstrLen = I - SubstrBegin;
-			BestSubstrBegin = SubstrBegin;
+			BestMatchBegin = MatchBegin;
+			BestMatchLen = MatchLen;
 		}
 	}
 	
-	const FString ValidData = Data.Mid(BestSubstrBegin, BestSubstrLen);
-	const FText ValidText = FText::FromString(ValidData);
+	const FText ValidText = FText::FromString(Data.Mid(BestMatchBegin, BestMatchLen));
 	ChangelistNumWidget->SetText(ValidText);
 }
 
