[BugFix] Crash when reviewing unloaded plugins/modules

diff --git a/Engine/Plugins/Editor/ChangelistReview/Source/Private/SSourceControlReview.cpp b/Engine/Plugins/Editor/ChangelistReview/Source/Private/SSourceControlReview.cpp
--- a/Engine/Plugins/Editor/ChangelistReview/Source/Private/SSourceControlReview.cpp
+++ b/Engine/Plugins/Editor/ChangelistReview/Source/Private/SSourceControlReview.cpp
@@ -665,7 +665,7 @@ SHeaderRow::FColumn::FArguments SSourceControlReview::HeaderColumn(FName HeaderN
 
 FString SSourceControlReview::AsAssetPath(const FString& FullCLPath)
 {
-	const FString ProjectName = FApp::GetProjectName();
+	const FString ProjectName = FString::Format(TEXT("/{0}/"), {FApp::GetProjectName()});
 	return UKismetSystemLibrary::GetProjectDirectory() / FullCLPath.RightChop(FullCLPath.Find(ProjectName) + ProjectName.Len());
 }
 
diff --git a/Engine/Plugins/Editor/ChangelistReview/Source/Private/SSourceControlReviewEntry.cpp b/Engine/Plugins/Editor/ChangelistReview/Source/Private/SSourceControlReviewEntry.cpp
--- a/Engine/Plugins/Editor/ChangelistReview/Source/Private/SSourceControlReviewEntry.cpp
+++ b/Engine/Plugins/Editor/ChangelistReview/Source/Private/SSourceControlReviewEntry.cpp
@@ -219,6 +219,13 @@ void SSourceControlReviewEntry::TryBindUAssetDiff()
 			{
 				DiffMethod.BindLambda([this, PreviousAsset, ReviewAsset]
 				{
+					const UBlueprint* ReviewBlueprint = Cast<UBlueprint>(ReviewAsset);
+					const UBlueprint* PreviousBlueprint = Cast<UBlueprint>(PreviousAsset);
+					if ((ReviewBlueprint && !ReviewBlueprint->ParentClass) || (PreviousBlueprint && !PreviousBlueprint->ParentClass))
+					{
+						FMessageDialog::Open(EAppMsgType::Ok, LOCTEXT("ChangelistNotFoundError", "This Blueprint is missing its parent class. Diff results may be incomplete.\n\nDo you need to load this file's plugin/module?\nFor more details, search 'Can't find file.' in the log"));
+					}
+					
 					FAssetToolsModule::GetModule().Get().DiffAssets(
 						PreviousAsset,
 						ReviewAsset,
diff --git a/Engine/Source/Editor/Kismet/Private/SBlueprintDiff.cpp b/Engine/Source/Editor/Kismet/Private/SBlueprintDiff.cpp
--- a/Engine/Source/Editor/Kismet/Private/SBlueprintDiff.cpp
+++ b/Engine/Source/Editor/Kismet/Private/SBlueprintDiff.cpp
@@ -849,7 +849,10 @@ void SBlueprintDiff::GenerateDifferencesList()
 	// Now that we have done the diffs, create the panel widgets
 	ModePanels.Add(MyBlueprintMode, GenerateMyBlueprintPanel());
 	ModePanels.Add(GraphMode, GenerateGraphPanel());
-	ModePanels.Add(DefaultsMode, GenerateDefaultsPanel());
+	if (PanelOld.Blueprint->ParentClass && PanelNew.Blueprint->ParentClass)
+	{
+		ModePanels.Add(DefaultsMode, GenerateDefaultsPanel());
+	}
 	ModePanels.Add(ClassSettingsMode, GenerateClassSettingsPanel());
 	if (bHasComponents)
 	{
