Closing the Editor while it's loading changes for the Changelist Review tool results in a crash

diff --git a/Engine/Plugins/Editor/ChangelistReview/Source/Private/SSourceControlReview.cpp b/Engine/Plugins/Editor/ChangelistReview/Source/Private/SSourceControlReview.cpp
--- a/Engine/Plugins/Editor/ChangelistReview/Source/Private/SSourceControlReview.cpp
+++ b/Engine/Plugins/Editor/ChangelistReview/Source/Private/SSourceControlReview.cpp
@@ -265,6 +265,15 @@ void SSourceControlReview::Construct(const FArguments& InArgs)
 	];
 }
 
+SSourceControlReview::~SSourceControlReview()
+{
+	if (GetChangelistDetailsCommand)
+	{
+		ISourceControlModule::Get().GetProvider().CancelOperation(GetChangelistDetailsCommand.ToSharedRef());
+		GetChangelistDetailsCommand.Reset();
+	}
+}
+
 void SSourceControlReview::LoadChangelist(const FString& Changelist)
 {
 	if (!IsLoading())
@@ -280,17 +289,23 @@ void SSourceControlReview::LoadChangelist(const FString& Changelist)
 	ChangelistFiles.Empty();
 
 	//This command runs p4 -describe (or similar for other version controls) to retrieve changelist record information
-	const TSharedRef<FGetChangelistDetails> GetChangelistDetailsCommand = ISourceControlOperation::Create<FGetChangelistDetails>();
+	GetChangelistDetailsCommand = ISourceControlOperation::Create<FGetChangelistDetails>();
 	GetChangelistDetailsCommand->SetChangelistNumber(Changelist);
 
 	ISourceControlModule::Get().GetProvider().Execute(
-		GetChangelistDetailsCommand,
+		GetChangelistDetailsCommand.ToSharedRef(),
 		EConcurrency::Asynchronous,
 		FSourceControlOperationComplete::CreateRaw(this, &SSourceControlReview::OnChangelistLoadComplete, Changelist));
 }
 
 void SSourceControlReview::OnChangelistLoadComplete(const FSourceControlOperationRef& InOperation, ECommandResult::Type InResult, FString Changelist)
 {
+	// this command is cancelled when this widget is destroyed. Exit immediately to avoid touching invalid data
+	if (InResult == ECommandResult::Cancelled)
+	{
+		return;
+	}
+	
 	const TSharedRef<FGetChangelistDetails, ESPMode::ThreadSafe> Operation = StaticCastSharedRef<FGetChangelistDetails>(InOperation);
 	const TArray<TMap<FString, FString>>& Record = Operation->GetChangelistDetails();
 	
diff --git a/Engine/Plugins/Editor/ChangelistReview/Source/Private/SSourceControlReview.h b/Engine/Plugins/Editor/ChangelistReview/Source/Private/SSourceControlReview.h
--- a/Engine/Plugins/Editor/ChangelistReview/Source/Private/SSourceControlReview.h
+++ b/Engine/Plugins/Editor/ChangelistReview/Source/Private/SSourceControlReview.h
@@ -9,6 +9,7 @@
 #include "Widgets/Views/SListView.h"
 #include "Widgets/Input/SEditableTextBox.h"
 
+class FGetChangelistDetails;
 class SSourceControlReviewEntry;
 class SChangelistEditableTextBox;
 class SProgressBar;
@@ -110,6 +111,8 @@ public:
 	/** Constructs the widget */
 	void Construct(const FArguments& InArgs);
 
+	virtual ~SSourceControlReview() override;
+
 	/**
 	 * Pulls up changelist record from source control
 	 * @param Changelist Changelist number to get
@@ -177,6 +180,7 @@ private:
 	uint32 FilesToLoad = 0;
 	uint32 FilesLoaded = 0;
 	TArray<TSharedPtr<FChangelistFileData>> ChangelistFiles;
+	TSharedPtr<FGetChangelistDetails> GetChangelistDetailsCommand;
 	
 	TSharedPtr<SEditableTextBox> ChangelistNumWidget;
 	TSharedPtr<STextBlock> EnterChangelistTextBlock;
