[Feature] Display redirectors in reviews in a way that reflects them as a rename
  - renamed/moved assets from within the editor will now show up as a single review entry

[BugFix] Assertion fails when loading changelist with Rename/Move change in the Changelist Review Tool. (specifically a Rename/Move that wasn't performed within UE)
  - We can only provide limited support for this case because renaming outside UE won't set up redirectors and the asset will keep it's old name dispite the changed filename

diff --git a/Engine/Plugins/Editor/ChangelistReview/Source/Private/SSourceControlReview.cpp b/Engine/Plugins/Editor/ChangelistReview/Source/Private/SSourceControlReview.cpp
--- a/Engine/Plugins/Editor/ChangelistReview/Source/Private/SSourceControlReview.cpp
+++ b/Engine/Plugins/Editor/ChangelistReview/Source/Private/SSourceControlReview.cpp
@@ -287,6 +287,7 @@ void SSourceControlReview::LoadChangelist(const FString& Changelist)
 	}
 
 	ChangelistFiles.Empty();
+	RedirectorsFound.Empty();
 
 	//This command runs p4 -describe (or similar for other version controls) to retrieve changelist record information
 	GetChangelistDetailsCommand = ISourceControlOperation::Create<FGetChangelistDetails>();
@@ -458,10 +459,45 @@ void SSourceControlReview::OnGetFileFromSourceControl(TSharedPtr<FChangelistFile
 
 	LoadingProgressBar->SetPercent(FilesToLoad ? static_cast<float>(FilesLoaded) / static_cast<float>(FilesToLoad) : 1.f);
 
+	if (UPackage* ReviewFilePkg = LoadPackage(nullptr, *ChangelistFileData->ReviewFileName, LOAD_ForDiff | LOAD_DisableCompileOnLoad | LOAD_DisableEngineVersionChecks))
+	{
+		const UObject* ReviewAsset = FindObject<UObject>(ReviewFilePkg, *ChangelistFileData->AssetName);
+		if(ReviewAsset && ReviewAsset->IsA<UObjectRedirector>())
+		{
+			const UPackage* RedirectedPackage = Cast<UObjectRedirector>(ReviewAsset)->DestinationObject->GetPackage();
+			const FString RedirectPath = IFileManager::Get().ConvertToRelativePath(*RedirectedPackage->GetLoadedPath().GetLocalFullPath());
+			RedirectorsFound.Add(RedirectPath, ChangelistFileData);
+		}
+	}
+
 	if (FilesToLoad == FilesLoaded)
 	{
 		const FAssetRegistryModule& AssetRegistryModule = FModuleManager::LoadModuleChecked<FAssetRegistryModule>( TEXT( "AssetRegistry" ) );
 
+		// Patch up redirectors such that the renamed assets only use a single diff entry
+		ChangelistFiles.RemoveAllSwap([&RedirectorsFound = RedirectorsFound](const TSharedPtr<FChangelistFileData> &ChangelistFile)
+		{
+			if (ChangelistFile->FileSourceControlAction == ESourceControlAction::Add)
+			{
+				const FString RelativePath = IFileManager::Get().ConvertToRelativePath(*ChangelistFile->AssetFilePath);
+				if (const TWeakPtr<FChangelistFileData>* Found = RedirectorsFound.Find(RelativePath))
+				{
+					if (const TSharedPtr<FChangelistFileData>& MergedRenameEntry = Found->Pin())
+					{
+						MergedRenameEntry->ReviewFileName = ChangelistFile->ReviewFileName;
+						MergedRenameEntry->PreviousAssetName = MergedRenameEntry->AssetName;
+						MergedRenameEntry->AssetName = ChangelistFile->AssetName;
+						MergedRenameEntry->RelativeFilePath = ChangelistFile->RelativeFilePath;
+						MergedRenameEntry->AssetFilePath = ChangelistFile->AssetFilePath;
+						MergedRenameEntry->PreviousFileRevisionNum = MergedRenameEntry->ReviewFileRevisionNum;
+						MergedRenameEntry->ReviewFileRevisionNum = ChangelistFile->ReviewFileRevisionNum;
+						return true;
+					}
+				}
+			}
+			return false;
+		});
+		
 		Algo::Sort(ChangelistFiles, [](const TSharedPtr<FChangelistFileData> &A, const TSharedPtr<FChangelistFileData> &B)
 		{
 			return A->RelativeFilePath < B->RelativeFilePath;
@@ -580,6 +616,14 @@ void SSourceControlReview::SetFileSourceControlAction(TSharedPtr<FChangelistFile
 	{
 		ChangelistFileData->FileSourceControlAction = ESourceControlAction::Integrate;
 	}
+	else if (SourceControlAction == FString(TEXT("move/delete")))
+	{
+		ChangelistFileData->FileSourceControlAction = ESourceControlAction::Delete;
+	}
+	else if (SourceControlAction == FString(TEXT("move/add")))
+	{
+		ChangelistFileData->FileSourceControlAction = ESourceControlAction::Add;
+	}
 	else
 	{
 		UE_LOG(LogSourceControl, Error, TEXT("Unable to parse revision control action information. '%s' diff will not be shown"), *ChangelistFileData->AssetName);
diff --git a/Engine/Plugins/Editor/ChangelistReview/Source/Private/SSourceControlReview.h b/Engine/Plugins/Editor/ChangelistReview/Source/Private/SSourceControlReview.h
--- a/Engine/Plugins/Editor/ChangelistReview/Source/Private/SSourceControlReview.h
+++ b/Engine/Plugins/Editor/ChangelistReview/Source/Private/SSourceControlReview.h
@@ -9,7 +9,6 @@
 #include "Widgets/Views/SListView.h"
 #include "Widgets/Input/SEditableTextBox.h"
 
-class FGetChangelistDetails;
 class SSourceControlReviewEntry;
 class SChangelistEditableTextBox;
 class SProgressBar;
@@ -47,6 +46,7 @@ namespace SourceControlReview
 			: AssetName(InAssetName)
 			, ReviewFileName(InReviewFilePkgName)
 			, ReviewFileRevisionNum(InReviewFileRevisionNum)
+			, PreviousAssetName() // assume no rename by default
 			, PreviousFileName(InPreviousFilePkgName)
 			, PreviousFileRevisionNum(InPreviousFileRevisionNum)
 		{
@@ -66,6 +66,8 @@ namespace SourceControlReview
 
 		FDateTime ReviewFileDateTime;
 
+		FString PreviousAssetName;
+
 		FString PreviousFileName;
 
 		FString PreviousFileRevisionNum;
@@ -180,6 +182,7 @@ private:
 	uint32 FilesToLoad = 0;
 	uint32 FilesLoaded = 0;
 	TArray<TSharedPtr<FChangelistFileData>> ChangelistFiles;
+	TMap<FString, TWeakPtr<FChangelistFileData>> RedirectorsFound;
 	TSharedPtr<FGetChangelistDetails> GetChangelistDetailsCommand;
 	
 	TSharedPtr<SEditableTextBox> ChangelistNumWidget;
diff --git a/Engine/Plugins/Editor/ChangelistReview/Source/Private/SSourceControlReviewEntry.cpp b/Engine/Plugins/Editor/ChangelistReview/Source/Private/SSourceControlReviewEntry.cpp
--- a/Engine/Plugins/Editor/ChangelistReview/Source/Private/SSourceControlReviewEntry.cpp
+++ b/Engine/Plugins/Editor/ChangelistReview/Source/Private/SSourceControlReviewEntry.cpp
@@ -209,67 +209,76 @@ bool SSourceControlReviewEntry::CanDiff() const
 
 void SSourceControlReviewEntry::TryBindUAssetDiff()
 {
+	UObject* ReviewAsset = nullptr;
 	if (UPackage* ReviewFilePkg = LoadPackage(nullptr, *ChangelistFileData.ReviewFileName, LOAD_ForDiff | LOAD_DisableCompileOnLoad | LOAD_DisableEngineVersionChecks))
 	{
-		if (UObject* ReviewAsset = FindObject<UObject>(ReviewFilePkg, *ChangelistFileData.AssetName))
+		ReviewAsset = FindObject<UObject>(ReviewFilePkg, *ChangelistFileData.AssetName);
+		if(ReviewAsset && ReviewAsset->IsA<UObjectRedirector>())
 		{
-			UPackage* PreviousFilePkg = LoadPackage(nullptr, *ChangelistFileData.PreviousFileName, LOAD_ForDiff | LOAD_DisableCompileOnLoad | LOAD_DisableEngineVersionChecks);
+			ReviewAsset = Cast<UObjectRedirector>(ReviewAsset)->DestinationObject;
+		}
+	}
+	UObject* PreviousAsset = nullptr;
+	if (UPackage* PreviousFilePkg = LoadPackage(nullptr, *ChangelistFileData.PreviousFileName, LOAD_ForDiff | LOAD_DisableCompileOnLoad | LOAD_DisableEngineVersionChecks))
+	{
+		if (ChangelistFileData.PreviousAssetName.IsEmpty())
+		{
+			PreviousAsset = FindObject<UObject>(PreviousFilePkg, *ChangelistFileData.AssetName);
+		}
+		else
+		{
+			PreviousAsset = FindObject<UObject>(PreviousFilePkg, *ChangelistFileData.PreviousAssetName);
+		}
+	}
+	
+	if (!ReviewAsset && !PreviousAsset)
+	{
+		return;
+	}
 
-			if (UObject* PreviousAsset = PreviousFilePkg ? FindObject<UObject>(PreviousFilePkg, *ChangelistFileData.AssetName) : nullptr)
-			{
-				DiffMethod.BindLambda([this, PreviousAsset, ReviewAsset]
-				{
-					const UBlueprint* ReviewBlueprint = Cast<UBlueprint>(ReviewAsset);
-					const UBlueprint* PreviousBlueprint = Cast<UBlueprint>(PreviousAsset);
-					if ((ReviewBlueprint && !ReviewBlueprint->ParentClass) || (PreviousBlueprint && !PreviousBlueprint->ParentClass))
-					{
-						FMessageDialog::Open(EAppMsgType::Ok, LOCTEXT("ChangelistNotFoundError", "This Blueprint is missing its parent class. Diff results may be incomplete.\n\nDo you need to load this file's plugin/module?\nFor more details, search 'Can't find file.' in the log"));
-					}
-					
-					FAssetToolsModule::GetModule().Get().DiffAssets(
-						PreviousAsset,
-						ReviewAsset,
-						GetPreviousFileRevisionInfo(),
-						GetReviewFileRevisionInfo()
-					);
-				});
-				
-			}
-			else
+	auto MakeDefaultAsset = [this](UObject* OtherAsset)->UObject*
+	{
+		if (const UBlueprint* BlueprintReviewAsset = Cast<UBlueprint>(OtherAsset))
+		{
+			return GetOrCreateBlueprintForDiff(BlueprintReviewAsset->GeneratedClass, BlueprintReviewAsset->BlueprintType);
+		}
+		else
+		{
+			//In case when file is not blueprint we are falling back on trying to create a UObject 
+			const FName EmptyObjectName = OtherAsset->GetFName();
+			const UClass* EmptyObjectAssetClass = OtherAsset->GetClass();
+			return NewObject<UObject>(OtherAsset, EmptyObjectAssetClass, EmptyObjectName, OtherAsset->GetFlags());
+		}
+	};
+
+	if (!ReviewAsset)
+	{
+		ReviewAsset = MakeDefaultAsset(PreviousAsset);
+	}
+
+	if (!PreviousAsset)
+	{
+		PreviousAsset = MakeDefaultAsset(ReviewAsset);
+	}
+
+	if (ReviewAsset && PreviousAsset)
+	{
+		DiffMethod.BindLambda([this, PreviousAsset, ReviewAsset]
+		{
+			const UBlueprint* ReviewBlueprint = Cast<UBlueprint>(ReviewAsset);
+			const UBlueprint* PreviousBlueprint = Cast<UBlueprint>(PreviousAsset);
+			if ((ReviewBlueprint && !ReviewBlueprint->ParentClass) || (PreviousBlueprint && !PreviousBlueprint->ParentClass))
 			{
-				if (const UBlueprint* BlueprintReviewAsset = Cast<UBlueprint>(ReviewAsset))
-				{
-					DiffMethod.BindLambda([this, BlueprintReviewAsset, ReviewAsset]
-					{
-						FAssetToolsModule::GetModule().Get().DiffAssets(
-							GetOrCreateBlueprintForDiff(BlueprintReviewAsset->GeneratedClass, BlueprintReviewAsset->BlueprintType),
-							ReviewAsset,
-							GetPreviousFileRevisionInfo(),
-							GetReviewFileRevisionInfo()
-						);
-					});
-				}
-				else
-				{
-					//In case when file is not blueprint we are falling back on trying to create a UObject 
-					const FName EmptyObjectName = ReviewAsset->GetFName();
-					const UClass* EmptyObjectAssetClass = ReviewAsset->GetClass();
-					
-					if (UObject* EmptyObject = NewObject<UObject>(ReviewAsset, EmptyObjectAssetClass, EmptyObjectName, ReviewAsset->GetFlags()))
-					{
-						DiffMethod.BindLambda([this, EmptyObject, ReviewAsset]
-						{
-							FAssetToolsModule::GetModule().Get().DiffAssets(
-								EmptyObject,
-								ReviewAsset,
-								GetPreviousFileRevisionInfo(),
-								GetReviewFileRevisionInfo()
-							);
-						});
-					}
-				}
+				FMessageDialog::Open(EAppMsgType::Ok, LOCTEXT("ChangelistNotFoundError", "This Blueprint is missing its parent class. Diff results may be incomplete.\n\nDo you need to load this file's plugin/module?\nFor more details, search 'Can't find file.' in the log"));
 			}
-		}
+			
+			FAssetToolsModule::GetModule().Get().DiffAssets(
+				PreviousAsset,
+				ReviewAsset,
+				GetPreviousFileRevisionInfo(),
+				GetReviewFileRevisionInfo()
+			);
+		});
 	}
 }
 
@@ -289,7 +298,12 @@ FRevisionInfo SSourceControlReviewEntry::GetReviewFileRevisionInfo() const
 	ReviewFileRevisionInfo.Changelist = ChangelistFileData.ChangelistNum;
 	ReviewFileRevisionInfo.Date = ChangelistFileData.ReviewFileDateTime;
 
-	if (ChangelistFileData.ChangelistState == EChangelistState::Pending)
+	if (ChangelistFileData.FileSourceControlAction == ESourceControlAction::Delete)
+	{
+		// new revision for removed files don't exist
+		ReviewFileRevisionInfo.Revision = TEXT("Does Not Exist");
+	}
+	else if (ChangelistFileData.ChangelistState == EChangelistState::Pending)
 	{
 		ReviewFileRevisionInfo.Revision = TEXT("Pending");
 	}
@@ -305,7 +319,20 @@ FRevisionInfo SSourceControlReviewEntry::GetPreviousFileRevisionInfo() const
 {
 	//We need to have valid revision data for some DiffAssets implementations (Although for now we don't have full data on previous file we are showing correct previous revision information)
 	FRevisionInfo PreviousFileRevisionInfo;
-	PreviousFileRevisionInfo.Revision = ChangelistFileData.PreviousFileRevisionNum.IsEmpty() ? TEXT("0") : ChangelistFileData.PreviousFileRevisionNum;
+	if (ChangelistFileData.FileSourceControlAction == ESourceControlAction::Add)
+	{
+		// previous revision for added files don't exist
+		PreviousFileRevisionInfo.Revision = TEXT("Does Not Exist");
+	}
+	else if (ChangelistFileData.PreviousFileRevisionNum.IsEmpty())
+	{
+		PreviousFileRevisionInfo.Revision = TEXT("0");
+	}
+	else
+	{
+		PreviousFileRevisionInfo.Revision = ChangelistFileData.PreviousFileRevisionNum;
+	}
+	
 	return PreviousFileRevisionInfo;
 }
 
@@ -381,6 +408,14 @@ FText SSourceControlReviewEntry::GetAssetType()
 
 FText SSourceControlReviewEntry::GetAssetNameText() const
 {
+	if (!ChangelistFileData.PreviousAssetName.IsEmpty())
+	{
+		return FText::Format(
+			LOCTEXT("RenamedAssetFormat", "{0} - (Renamed from: {1})"),
+			FText::FromString(ChangelistFileData.AssetName),
+			FText::FromString(ChangelistFileData.PreviousAssetName)
+			);
+	}
 	return FText::FromString(ChangelistFileData.AssetName);
 }
 
