Non Unity Build errors Fix

diff --git a/Engine/Plugins/Editor/ChangelistReview/Source/Private/SSourceControlReview.cpp b/Engine/Plugins/Editor/ChangelistReview/Source/Private/SSourceControlReview.cpp
--- a/Engine/Plugins/Editor/ChangelistReview/Source/Private/SSourceControlReview.cpp
+++ b/Engine/Plugins/Editor/ChangelistReview/Source/Private/SSourceControlReview.cpp
@@ -2,6 +2,7 @@
 
 #include "SSourceControlReview.h"
 
+#include "HAL/FileManager.h"
 #include "ClassIconFinder.h"
 #include "ISourceControlModule.h"
 #include "ISourceControlProvider.h"
