fix FGetChangelistDetails undeclaired errors

diff --git a/Engine/Plugins/Editor/ChangelistReview/Source/Private/SSourceControlReview.h b/Engine/Plugins/Editor/ChangelistReview/Source/Private/SSourceControlReview.h
--- a/Engine/Plugins/Editor/ChangelistReview/Source/Private/SSourceControlReview.h
+++ b/Engine/Plugins/Editor/ChangelistReview/Source/Private/SSourceControlReview.h
@@ -3,6 +3,7 @@
 #pragma once
 #include "ISourceControlOperation.h"
 #include "ISourceControlProvider.h"
+#include "SourceControlOperations.h"
 #include "Widgets/SCompoundWidget.h"
 #include "Widgets/Text/STextBlock.h"
 #include "Widgets/SBoxPanel.h"
