[Feature] Adding CL history dropdown to the CL input in the blueprint review tool

diff --git a/Engine/Plugins/Editor/ChangelistReview/Source/Private/SSourceControlReview.cpp b/Engine/Plugins/Editor/ChangelistReview/Source/Private/SSourceControlReview.cpp
--- a/Engine/Plugins/Editor/ChangelistReview/Source/Private/SSourceControlReview.cpp
+++ b/Engine/Plugins/Editor/ChangelistReview/Source/Private/SSourceControlReview.cpp
@@ -14,7 +14,7 @@
 #include "AssetRegistry/AssetRegistryModule.h"
 #include "Kismet/KismetSystemLibrary.h"
 #include "Styling/StarshipCoreStyle.h"
-#include "Widgets/Input/SEditableTextBox.h"
+#include "Widgets/Input/SEditableText.h"
 #include "Widgets/Input/SMultiLineEditableTextBox.h"
 #include "Widgets/Notifications/SProgressBar.h"
 #include "Widgets/Input/SButton.h"
@@ -22,10 +22,10 @@
 #include "Misc/MessageDialog.h"
 #include "Styling/StyleColors.h"
 #include "Widgets/Layout/SGridPanel.h"
-#include "AssetRegistry/AssetRegistryModule.h"
 #include "Brushes/SlateRoundedBoxBrush.h"
 #include "Internationalization/Regex.h"
 #include "Widgets/Views/SListView.h"
+#include "Widgets/Input/SComboBox.h"
 
 #define LOCTEXT_NAMESPACE "SourceControlReview"
 
@@ -81,6 +81,8 @@ void SSourceControlReview::Construct(const FArguments& InArgs)
 		.SetBackgroundImageReadOnly(RecessedBrush);
 
 	const static FMargin InfoWidgetMargin(4.f, 2.f, 4.f, 8.f);
+
+	LoadCLHistory();
 	
 	ChildSlot
 	[
@@ -114,13 +116,21 @@ void SSourceControlReview::Construct(const FArguments& InArgs)
 					.Padding(0.f, 0.f, 8.f, 0.f)
 					.AutoWidth()
 					[
-						SAssignNew(ChangelistNumWidget, SEditableTextBox)
-						.Font(FStyleFonts::Get().Normal)
-						.MinDesiredWidth(55.f)
-						.Justification(ETextJustify::Center)
-						.OnTextCommitted(this, &SSourceControlReview::OnChangelistNumCommitted)
-						.OnTextChanged(this, &SSourceControlReview::OnChangelistNumChanged)
-						.BackgroundColor(FStyleColors::Input)
+						SAssignNew(ChangelistNumComboBox, SComboBox<TSharedPtr<FChangelistLightInfo>>)
+						.OptionsSource(&CLHistory)
+						.OnGenerateWidget(this, &SSourceControlReview::MakeCLComboOption)
+						.OnSelectionChanged(this, &SSourceControlReview::OnCLComboSelection)
+						.ContentPadding(0.f)
+						.Content()
+						[
+							SAssignNew(ChangelistNumText, SEditableText)
+                     		.Font(FStyleFonts::Get().Normal)
+                     		.MinDesiredWidth(55.f)
+                     		.Justification(ETextJustify::Center)
+                     		.OnTextCommitted(this, &SSourceControlReview::OnChangelistNumCommitted)
+                     		.OnTextChanged(this, &SSourceControlReview::OnChangelistNumChanged)
+						]
+						
 					]
 					+SHorizontalBox::Slot()
 					.VAlign(VAlign_Center)
@@ -289,6 +299,15 @@ void SSourceControlReview::LoadChangelist(const FString& Changelist)
 		return;
 	}
 
+	for(int32 I = 1; I < CLHistory.Num(); ++I)
+	{
+		if (CLHistory[I]->Number.EqualTo(CLHistory[0]->Number))
+		{
+			CLHistory.RemoveAt(I);
+			break;
+		}
+	}
+	bUncommittedChangelistNum = false;
 	ChangelistFiles.Empty();
 	RedirectorsFound.Empty();
 
@@ -315,6 +334,7 @@ void SSourceControlReview::OnChangelistLoadComplete(const FSourceControlOperatio
 	
 	if (!IsChangelistRecordValid(Record))
 	{
+		CLHistory.RemoveAt(0); // don't save history of invalid changelists
 		SetLoading(false);
 		return;
 	}
@@ -339,6 +359,13 @@ void SSourceControlReview::OnChangelistLoadComplete(const FSourceControlOperatio
 	FString RecordActionMapKey = ReviewHelpers::FileActionKey + RecordFileIndexStr;
 
 	SetChangelistInfo(ChangelistRecord);
+	if (CLHistory[0]->Author.IsEmpty())
+	{
+		// annoyingly the options list doesn't visually update unless the shared pointer is different
+		CLHistory[0] = MakeShared<FChangelistLightInfo>(CLHistory[0]->Number, CurrentChangelistInfo.Author, CurrentChangelistInfo.Description);
+	}
+	ChangelistNumComboBox->RefreshOptions();
+	SaveCLHistory();
 	
 	//the loop checks if we have a valid record "depotFile(Index)" in the records to add a file entry
 	while (ChangelistRecord.Contains(RecordFileMapKey) && ChangelistRecord.Contains(RecordRevisionMapKey))
@@ -412,7 +439,7 @@ void SSourceControlReview::OnChangelistLoadComplete(const FSourceControlOperatio
 
 FReply SSourceControlReview::OnLoadChangelistClicked()
 {
-	LoadChangelist(ChangelistNumWidget->GetText().ToString());
+	LoadChangelist(ChangelistNumText->GetText().ToString());
 	return FReply::Handled();
 }
 
@@ -440,7 +467,27 @@ void SSourceControlReview::OnChangelistNumChanged(const FText& Text)
 	}
 	
 	const FText ValidText = FText::FromString(Data.Mid(BestMatchBegin, BestMatchLen));
-	ChangelistNumWidget->SetText(ValidText);
+	ChangelistNumText->SetText(ValidText);
+
+	const TSharedPtr<FChangelistLightInfo> ChangelistLightInfo = MakeShared<FChangelistLightInfo>(ValidText);
+	if (!bUncommittedChangelistNum || CLHistory.IsEmpty())
+	{
+		bUncommittedChangelistNum = true;
+		if (CLHistory.Num() > 5)
+		{
+			CLHistory.Pop();
+		}
+		CLHistory.Insert(ChangelistLightInfo, 0);
+	}
+	else
+	{
+		CLHistory[0] = ChangelistLightInfo;
+	}
+	ChangelistNumComboBox->RefreshOptions();
+	if (ChangelistNumComboBox->GetSelectedItem() != CLHistory[0])
+	{
+		ChangelistNumComboBox->SetSelectedItem(CLHistory[0]);
+	}
 }
 
 void SSourceControlReview::OnChangelistNumCommitted(const FText& Text, ETextCommit::Type CommitMethod)
@@ -451,6 +498,67 @@ void SSourceControlReview::OnChangelistNumCommitted(const FText& Text, ETextComm
 	}
 }
 
+TSharedRef<SWidget> SSourceControlReview::MakeCLComboOption(TSharedPtr<FChangelistLightInfo> Item) const
+{
+	FText Text;
+	if (!Item->Author.IsEmpty())
+	{
+		Text = FText::Format(LOCTEXT("CLComboOption", "{0} - {1}"), Item->Number, Item->Author);
+	}
+	else
+	{
+		Text =Item->Number;
+	}
+	
+	
+	return SNew(STextBlock)
+		.Text(Text)
+		.ToolTipText(Item->Description);
+}
+
+void SSourceControlReview::OnCLComboSelection(TSharedPtr<FChangelistLightInfo> Item, ESelectInfo::Type SelectInfo) const
+{
+	if (Item.IsValid() && SelectInfo != ESelectInfo::Direct)
+	{
+		ChangelistNumText->SetText(Item->Number);
+	}
+}
+
+void SSourceControlReview::SaveCLHistory()
+{
+	TArray<FString> Numbers;
+	TArray<FString> Authors;
+	TArray<FString> Descriptions;
+	for(const TSharedPtr<FChangelistLightInfo>& Item : CLHistory)
+	{
+		Numbers.Add(Item->Number.ToString());
+		Authors.Add(Item->Author.ToString());
+		Descriptions.Add(Item->Description.ToString());
+	}
+	GConfig->SetArray(TEXT("SourceControlReview"), TEXT("CLHistory.Numbers"), Numbers, GEngineIni);
+	GConfig->SetArray(TEXT("SourceControlReview"), TEXT("CLHistory.Authors"), Authors, GEngineIni);
+	GConfig->SetArray(TEXT("SourceControlReview"), TEXT("CLHistory.Descriptions"), Descriptions, GEngineIni);
+}
+
+void SSourceControlReview::LoadCLHistory()
+{
+	TArray<FString> Numbers;
+	TArray<FString> Authors;
+	TArray<FString> Descriptions;
+	GConfig->GetArray(TEXT("SourceControlReview"), TEXT("CLHistory.Numbers"), Numbers, GEngineIni);
+	GConfig->GetArray(TEXT("SourceControlReview"), TEXT("CLHistory.Authors"), Authors, GEngineIni);
+	GConfig->GetArray(TEXT("SourceControlReview"), TEXT("CLHistory.Descriptions"), Descriptions, GEngineIni);
+	CLHistory.Empty();
+	for(int32 I = 0; I < Numbers.Num(); ++I)
+	{
+		CLHistory.Add(MakeShared<FChangelistLightInfo>(
+			FText::FromString(Numbers[I]),
+			FText::FromString(Authors[I]),
+			FText::FromString(Descriptions[I])
+		));
+	}
+}
+
 void SSourceControlReview::OnGetFileFromSourceControl(TSharedPtr<FChangelistFileData> ChangelistFileData)
 {
 	if (ChangelistFileData && ChangelistFileData->IsDataValidForEntry())
diff --git a/Engine/Plugins/Editor/ChangelistReview/Source/Private/SSourceControlReview.h b/Engine/Plugins/Editor/ChangelistReview/Source/Private/SSourceControlReview.h
--- a/Engine/Plugins/Editor/ChangelistReview/Source/Private/SSourceControlReview.h
+++ b/Engine/Plugins/Editor/ChangelistReview/Source/Private/SSourceControlReview.h
@@ -88,13 +88,27 @@ namespace SourceControlReview
 	private:
 		TOptional<const UClass*> CachedIconClass;
 	};
-
+	
 	namespace ColumnIds
 	{
 		inline const FName Status = TEXT("Status");
 		inline const FName File = TEXT("File");
 		inline const FName Tools = TEXT("Tools");
 	}
+
+	struct FChangelistLightInfo
+	{
+		explicit FChangelistLightInfo(const FText& Number)
+			: Number(Number)
+		{}
+		FChangelistLightInfo(const FText& Number, const FText& Author, const FText& Description)
+			: Number(Number), Author(Author), Description(Description)
+		{}
+		
+		FText Number;
+		FText Author;
+		FText Description;
+	};
 }
 
 
@@ -127,6 +141,11 @@ private:
 	FReply OnLoadChangelistClicked();
 	void OnChangelistNumChanged(const FText& Text);
 	void OnChangelistNumCommitted(const FText& Text, ETextCommit::Type CommitMethod);
+
+	TSharedRef<SWidget> MakeCLComboOption(TSharedPtr<SourceControlReview::FChangelistLightInfo>Item) const;
+	void OnCLComboSelection(TSharedPtr<SourceControlReview::FChangelistLightInfo> Item, ESelectInfo::Type SelectInfo) const;
+	void SaveCLHistory();
+	void LoadCLHistory();
 	
 	/**
 	 * Called when file is retrieved from
@@ -185,8 +204,11 @@ private:
 	TArray<TSharedPtr<FChangelistFileData>> ChangelistFiles;
 	TMap<FString, TWeakPtr<FChangelistFileData>> RedirectorsFound;
 	TSharedPtr<FGetChangelistDetails> GetChangelistDetailsCommand;
+	TArray<TSharedPtr<SourceControlReview::FChangelistLightInfo>> CLHistory;
+	bool bUncommittedChangelistNum = false;
 	
-	TSharedPtr<SEditableTextBox> ChangelistNumWidget;
+	TSharedPtr<SComboBox<TSharedPtr<SourceControlReview::FChangelistLightInfo>>> ChangelistNumComboBox;
+	TSharedPtr<SEditableText> ChangelistNumText;
 	TSharedPtr<STextBlock> EnterChangelistTextBlock;
 	TSharedPtr<STextBlock> LoadingTextBlock;
 	TSharedPtr<SProgressBar> LoadingProgressBar;
