[HotFix] Unity build error fix

diff --git a/Engine/Plugins/Editor/ChangelistReview/Source/Private/SSourceControlReview.cpp b/Engine/Plugins/Editor/ChangelistReview/Source/Private/SSourceControlReview.cpp
--- a/Engine/Plugins/Editor/ChangelistReview/Source/Private/SSourceControlReview.cpp
+++ b/Engine/Plugins/Editor/ChangelistReview/Source/Private/SSourceControlReview.cpp
@@ -16,6 +16,7 @@
 #include "Styling/StarshipCoreStyle.h"
 #include "Widgets/Input/SEditableText.h"
 #include "Widgets/Input/SMultiLineEditableTextBox.h"
+#include "Widgets/Input/SEditableTextBox.h"
 #include "Widgets/Notifications/SProgressBar.h"
 #include "Widgets/Input/SButton.h"
 #include "Misc/App.h"
@@ -26,6 +27,8 @@
 #include "Internationalization/Regex.h"
 #include "Widgets/Views/SListView.h"
 #include "Widgets/Input/SComboBox.h"
+#include "CoreGlobals.h"
+#include "Misc/ConfigCacheIni.h"
 
 #define LOCTEXT_NAMESPACE "SourceControlReview"
 
diff --git a/Engine/Plugins/Editor/ChangelistReview/Source/Private/SSourceControlReview.h b/Engine/Plugins/Editor/ChangelistReview/Source/Private/SSourceControlReview.h
--- a/Engine/Plugins/Editor/ChangelistReview/Source/Private/SSourceControlReview.h
+++ b/Engine/Plugins/Editor/ChangelistReview/Source/Private/SSourceControlReview.h
@@ -4,8 +4,9 @@
 #include "ISourceControlProvider.h"
 #include "SourceControlOperations.h"
 #include "Widgets/Views/SHeaderRow.h"
+#include "Widgets/Input/SComboBox.h"
 
-class SEditableTextBox;
+class SEditableText;
 class STableViewBase;
 class STextBlock;
 template <typename ItemType> class SListView;
