[CrashFix] Crash occurs when reviewing changelist with Branch changes in the Changelist Review Tool

diff --git a/Engine/Plugins/Editor/ChangelistReview/Source/Private/SSourceControlReview.cpp b/Engine/Plugins/Editor/ChangelistReview/Source/Private/SSourceControlReview.cpp
--- a/Engine/Plugins/Editor/ChangelistReview/Source/Private/SSourceControlReview.cpp
+++ b/Engine/Plugins/Editor/ChangelistReview/Source/Private/SSourceControlReview.cpp
@@ -415,7 +415,7 @@ void SSourceControlReview::OnChangelistLoadComplete(const FSourceControlOperatio
 		
 		ISourceControlModule::Get().GetProvider().Execute(GetFileToReviewCommand, EConcurrency::Asynchronous, FSourceControlOperationComplete::CreateLambda(GetFileCommandResponse, &ChangelistFileData->ReviewFileName));
 
-		if (ChangelistFileData->FileSourceControlAction != ESourceControlAction::Add)
+		if (ChangelistFileData->FileSourceControlAction != ESourceControlAction::Add && ChangelistFileData->FileSourceControlAction != ESourceControlAction::Branch)
 		{
 			FilesToLoad++;
 
diff --git a/Engine/Plugins/Editor/ChangelistReview/Source/Private/SSourceControlReview.h b/Engine/Plugins/Editor/ChangelistReview/Source/Private/SSourceControlReview.h
--- a/Engine/Plugins/Editor/ChangelistReview/Source/Private/SSourceControlReview.h
+++ b/Engine/Plugins/Editor/ChangelistReview/Source/Private/SSourceControlReview.h
@@ -57,7 +57,7 @@ namespace SourceControlReview
 		bool IsDataValidForEntry() const
 		{
 			return FileSourceControlAction != ESourceControlAction::Unset && !ReviewFileName.IsEmpty() &&
-				(!PreviousFileName.IsEmpty() || FileSourceControlAction == ESourceControlAction::Add);
+				(!PreviousFileName.IsEmpty() || FileSourceControlAction == ESourceControlAction::Add || FileSourceControlAction == ESourceControlAction::Branch);
 		}
 
 		FString AssetName;
