[CrashFix] Crash occurs when Reviewing Shelved Changelist with a File Moved In-Editor

diff --git a/Engine/Plugins/Editor/ChangelistReview/Source/Private/SSourceControlReview.cpp b/Engine/Plugins/Editor/ChangelistReview/Source/Private/SSourceControlReview.cpp
--- a/Engine/Plugins/Editor/ChangelistReview/Source/Private/SSourceControlReview.cpp
+++ b/Engine/Plugins/Editor/ChangelistReview/Source/Private/SSourceControlReview.cpp
@@ -574,9 +574,15 @@ void SSourceControlReview::OnGetFileFromSourceControl(TSharedPtr<FChangelistFile
 		const UObject* ReviewAsset = FindObject<UObject>(ReviewFilePkg, *ChangelistFileData->AssetName);
 		if(ReviewAsset && ReviewAsset->IsA<UObjectRedirector>())
 		{
-			const UPackage* RedirectedPackage = Cast<UObjectRedirector>(ReviewAsset)->DestinationObject->GetPackage();
-			const FString RedirectPath = IFileManager::Get().ConvertToRelativePath(*RedirectedPackage->GetLoadedPath().GetLocalFullPath());
-			RedirectorsFound.Add(RedirectPath, ChangelistFileData);
+			if (const UObjectRedirector* Redirector = Cast<UObjectRedirector>(ReviewAsset))
+			{
+				if (Redirector->DestinationObject)
+				{
+					const UPackage* RedirectedPackage = Cast<UObjectRedirector>(ReviewAsset)->DestinationObject->GetPackage();
+					const FString RedirectPath = IFileManager::Get().ConvertToRelativePath(*RedirectedPackage->GetLoadedPath().GetLocalFullPath());
+					RedirectorsFound.Add(RedirectPath, ChangelistFileData);
+				}
+			}
 		}
 	}
 
diff --git a/Engine/Plugins/Editor/ChangelistReview/Source/Private/SSourceControlReviewEntry.cpp b/Engine/Plugins/Editor/ChangelistReview/Source/Private/SSourceControlReviewEntry.cpp
--- a/Engine/Plugins/Editor/ChangelistReview/Source/Private/SSourceControlReviewEntry.cpp
+++ b/Engine/Plugins/Editor/ChangelistReview/Source/Private/SSourceControlReviewEntry.cpp
@@ -219,15 +219,18 @@ void SSourceControlReviewEntry::TryBindUAssetDiff()
 		}
 	}
 	UObject* PreviousAsset = nullptr;
-	if (UPackage* PreviousFilePkg = LoadPackage(nullptr, *ChangelistFileData.PreviousFileName, LOAD_ForDiff | LOAD_DisableCompileOnLoad | LOAD_DisableEngineVersionChecks))
+	if (!ChangelistFileData.PreviousFileName.IsEmpty())
 	{
-		if (ChangelistFileData.PreviousAssetName.IsEmpty())
+		if (UPackage* PreviousFilePkg = LoadPackage(nullptr, *ChangelistFileData.PreviousFileName, LOAD_ForDiff | LOAD_DisableCompileOnLoad | LOAD_DisableEngineVersionChecks))
 		{
-			PreviousAsset = FindObject<UObject>(PreviousFilePkg, *ChangelistFileData.AssetName);
-		}
-		else
-		{
-			PreviousAsset = FindObject<UObject>(PreviousFilePkg, *ChangelistFileData.PreviousAssetName);
+			if (ChangelistFileData.PreviousAssetName.IsEmpty())
+			{
+				PreviousAsset = FindObject<UObject>(PreviousFilePkg, *ChangelistFileData.AssetName);
+			}
+			else
+			{
+				PreviousAsset = FindObject<UObject>(PreviousFilePkg, *ChangelistFileData.PreviousAssetName);
+			}
 		}
 	}
 	
