[BugFix] Review tool doesn't properly merge redirectors when reviewing CLs that aren't already checked out

diff --git a/Engine/Plugins/Editor/ChangelistReview/Source/Private/SSourceControlReview.cpp b/Engine/Plugins/Editor/ChangelistReview/Source/Private/SSourceControlReview.cpp
--- a/Engine/Plugins/Editor/ChangelistReview/Source/Private/SSourceControlReview.cpp
+++ b/Engine/Plugins/Editor/ChangelistReview/Source/Private/SSourceControlReview.cpp
@@ -7,6 +7,7 @@
 #include "ClassIconFinder.h"
 #include "Framework/Views/TableViewMetadata.h"
 #include "ISourceControlModule.h"
+#include "SourceControlHelpers.h"
 #include "Misc/PackageName.h"
 #include "SourceControlOperations.h"
 #include "SSourceControlReviewEntry.h"
@@ -50,8 +51,8 @@ const UClass* FChangelistFileData::GetIconClass()
 	{
 		CachedIconClass = nullptr;
 		const FAssetRegistryModule& AssetRegistryModule = FModuleManager::LoadModuleChecked<FAssetRegistryModule>(TEXT("AssetRegistry"));
-		FString TempPackageName;
-		if (FPackageName::TryConvertFilenameToLongPackageName(ReviewFileName, TempPackageName))
+		const FString TempPackageName = ReviewFileTempPath.GetPackageName();
+		if (!TempPackageName.IsEmpty())
 		{
 			TArray<FAssetData> OutAssetData;
 			AssetRegistryModule.Get().GetAssetsByPackageName(*TempPackageName, OutAssetData);
@@ -308,7 +309,6 @@ void SSourceControlReview::LoadChangelist(const FString& Changelist)
 	}
 	bUncommittedChangelistNum = false;
 	ChangelistFiles.Empty();
-	RedirectorsFound.Empty();
 
 	//This command runs p4 -describe (or similar for other version controls) to retrieve changelist record information
 	GetChangelistDetailsCommand = ISourceControlOperation::Create<FGetChangelistDetails>();
@@ -370,7 +370,6 @@ void SSourceControlReview::OnChangelistLoadComplete(const FSourceControlOperatio
 	while (ChangelistRecord.Contains(RecordFileMapKey) && ChangelistRecord.Contains(RecordRevisionMapKey))
 	{
 		const FString &FileDepotPath = ChangelistRecord[RecordFileMapKey];
-		const FString AssetName = FPaths::GetBaseFilename(FileDepotPath, true);
 		
 		//For each 1 file we are loading 2 revisions so files to load is always incremented by two per file
 		FilesToLoad++; 
@@ -378,8 +377,11 @@ void SSourceControlReview::OnChangelistLoadComplete(const FSourceControlOperatio
 		const bool bIsShelved = ChangelistRecord[ReviewHelpers::ChangelistStatusKey] == ReviewHelpers::ChangelistPendingStatusKey;
 		const int32 AssetRevision = FCString::Atoi(*ChangelistRecord[RecordRevisionMapKey]);
 	
-		TSharedPtr<FChangelistFileData> ChangelistFileData = MakeShared<FChangelistFileData>(AssetName, TEXT(""), ChangelistRecord[RecordRevisionMapKey], TEXT(""), TEXT(""));
+		TSharedPtr<FChangelistFileData> ChangelistFileData = MakeShared<FChangelistFileData>();
 
+		ChangelistFileData->AssetName = FPaths::GetBaseFilename(FileDepotPath, true);
+		ChangelistFileData->ReviewFileRevisionNum = ChangelistRecord[RecordRevisionMapKey];
+		
 		ChangelistFileData->ReviewFileDateTime = FDateTime(1970, 1, 1, 0, 0, 0, 0) + FTimespan::FromSeconds(FCString::Atoi(*ChangelistRecord[ReviewHelpers::TimeKey]));
 		ChangelistFileData->ChangelistNum = FCString::Atoi(*Changelist);
 
@@ -387,7 +389,8 @@ void SSourceControlReview::OnChangelistLoadComplete(const FSourceControlOperatio
 		ChangelistFileData->ChangelistState = bIsShelved ? EChangelistState::Pending : EChangelistState::Submitted;
 
 		//Building absolute local path is needed to use local file to retrieve file history information from p4 to then show file revision data
-		ChangelistFileData->AssetFilePath = AsAssetPath(ChangelistRecord[RecordFileMapKey]);
+		ChangelistFileData->AssetDepotPath = FileDepotPath;
+		ChangelistFileData->AssetFilePath = AsAssetPath(ChangelistFileData->AssetDepotPath);
 		ChangelistFileData->RelativeFilePath = TrimSharedPath(ChangelistRecord[RecordFileMapKey]);
 
 		SetFileSourceControlAction(ChangelistFileData, ChangelistRecord[RecordActionMapKey]);
@@ -399,17 +402,17 @@ void SSourceControlReview::OnChangelistLoadComplete(const FSourceControlOperatio
 		TSharedRef<FGetFile> GetFileToReviewCommand = ISourceControlOperation::Create<FGetFile>(Changelist, ChangelistRecord[RecordRevisionMapKey], ChangelistRecord[RecordFileMapKey], bIsShelved);
 
 		TWeakPtr<SSourceControlReview> WeakReviewWidget = SharedThis(this);
-		const auto GetFileCommandResponse = [WeakReviewWidget, ChangelistFileData](const FSourceControlOperationRef& InOperation, ECommandResult::Type InResult, FString* OutFileName)
+		const auto GetFileCommandResponse = [WeakReviewWidget, ChangelistFileData](const FSourceControlOperationRef& InOperation, ECommandResult::Type InResult, FPackagePath* OutFilePath)
 		{
 			if (WeakReviewWidget.IsValid())
 			{
 				const TSharedRef<FGetFile, ESPMode::ThreadSafe> Operation = StaticCastSharedRef<FGetFile>(InOperation);
-				*OutFileName = Operation->GetOutPackageFilename();
+				*OutFilePath = FPackagePath::FromLocalPath(Operation->GetOutPackageFilename());
 				WeakReviewWidget.Pin()->OnGetFileFromSourceControl(ChangelistFileData);
 			}
 		};
 		
-		ISourceControlModule::Get().GetProvider().Execute(GetFileToReviewCommand, EConcurrency::Asynchronous, FSourceControlOperationComplete::CreateLambda(GetFileCommandResponse, &ChangelistFileData->ReviewFileName));
+		ISourceControlModule::Get().GetProvider().Execute(GetFileToReviewCommand, EConcurrency::Asynchronous, FSourceControlOperationComplete::CreateLambda(GetFileCommandResponse, &ChangelistFileData->ReviewFileTempPath));
 
 		if (ChangelistFileData->FileSourceControlAction != ESourceControlAction::Add && ChangelistFileData->FileSourceControlAction != ESourceControlAction::Branch)
 		{
@@ -417,7 +420,7 @@ void SSourceControlReview::OnChangelistLoadComplete(const FSourceControlOperatio
 
 			TSharedRef<FGetFile> GetPreviousFileCommand = ISourceControlOperation::Create<FGetFile>(Changelist, PreviousAssetRevisionStr, ChangelistRecord[RecordFileMapKey], false);
 
-			ISourceControlModule::Get().GetProvider().Execute(GetPreviousFileCommand, EConcurrency::Asynchronous, FSourceControlOperationComplete::CreateLambda(GetFileCommandResponse, &ChangelistFileData->PreviousFileName));
+			ISourceControlModule::Get().GetProvider().Execute(GetPreviousFileCommand, EConcurrency::Asynchronous, FSourceControlOperationComplete::CreateLambda(GetFileCommandResponse, &ChangelistFileData->PreviousFileTempPath));
 		}
 
 		//Update map keys with the next file index
@@ -569,50 +572,12 @@ void SSourceControlReview::OnGetFileFromSourceControl(TSharedPtr<FChangelistFile
 
 	LoadingProgressBar->SetPercent(FilesToLoad ? static_cast<float>(FilesLoaded) / static_cast<float>(FilesToLoad) : 1.f);
 
-	if (UPackage* ReviewFilePkg = LoadPackage(nullptr, *ChangelistFileData->ReviewFileName, LOAD_ForDiff | LOAD_DisableCompileOnLoad | LOAD_DisableEngineVersionChecks))
-	{
-		const UObject* ReviewAsset = FindObject<UObject>(ReviewFilePkg, *ChangelistFileData->AssetName);
-		if(ReviewAsset && ReviewAsset->IsA<UObjectRedirector>())
-		{
-			if (const UObjectRedirector* Redirector = Cast<UObjectRedirector>(ReviewAsset))
-			{
-				if (Redirector->DestinationObject)
-				{
-					const UPackage* RedirectedPackage = Cast<UObjectRedirector>(ReviewAsset)->DestinationObject->GetPackage();
-					const FString RedirectPath = IFileManager::Get().ConvertToRelativePath(*RedirectedPackage->GetLoadedPath().GetLocalFullPath());
-					RedirectorsFound.Add(RedirectPath, ChangelistFileData);
-				}
-			}
-		}
-	}
-
 	if (FilesToLoad == FilesLoaded)
 	{
 		const FAssetRegistryModule& AssetRegistryModule = FModuleManager::LoadModuleChecked<FAssetRegistryModule>( TEXT( "AssetRegistry" ) );
 
-		// Patch up redirectors such that the renamed assets only use a single diff entry
-		ChangelistFiles.RemoveAllSwap([&RedirectorsFound = RedirectorsFound](const TSharedPtr<FChangelistFileData> &ChangelistFile)
-		{
-			if (ChangelistFile->FileSourceControlAction == ESourceControlAction::Add)
-			{
-				const FString RelativePath = IFileManager::Get().ConvertToRelativePath(*ChangelistFile->AssetFilePath);
-				if (const TWeakPtr<FChangelistFileData>* Found = RedirectorsFound.Find(RelativePath))
-				{
-					if (const TSharedPtr<FChangelistFileData>& MergedRenameEntry = Found->Pin())
-					{
-						MergedRenameEntry->ReviewFileName = ChangelistFile->ReviewFileName;
-						MergedRenameEntry->PreviousAssetName = MergedRenameEntry->AssetName;
-						MergedRenameEntry->AssetName = ChangelistFile->AssetName;
-						MergedRenameEntry->RelativeFilePath = ChangelistFile->RelativeFilePath;
-						MergedRenameEntry->AssetFilePath = ChangelistFile->AssetFilePath;
-						MergedRenameEntry->PreviousFileRevisionNum = MergedRenameEntry->ReviewFileRevisionNum;
-						MergedRenameEntry->ReviewFileRevisionNum = ChangelistFile->ReviewFileRevisionNum;
-						return true;
-					}
-				}
-			}
-			return false;
-		});
+		// Merge asset moves and renames into single entries
+		FixupRedirectors();
 		
 		Algo::Sort(ChangelistFiles, [](const TSharedPtr<FChangelistFileData> &A, const TSharedPtr<FChangelistFileData> &B)
 		{
@@ -624,7 +589,7 @@ void SSourceControlReview::OnGetFileFromSourceControl(TSharedPtr<FChangelistFile
 		{
 			if(FileData)
 			{
-				ChangelistFilePaths.Add(FileData->ReviewFileName);
+				ChangelistFilePaths.Add(FileData->ReviewFileTempPath.GetLocalFullPath());
 			}
 		}
 		
@@ -805,6 +770,108 @@ FString SSourceControlReview::TrimSharedPath(FString FullCLPath) const
 	return FullCLPath;
 }
 
+static bool IsObjectRedirector(FLinkerLoad* LinkerLoad, FStringView AssetName)
+{
+	// If the asset's name matches the name of a redirector, we know the entire asset is a redirector
+	for (const FObjectExport& Export : LinkerLoad->ExportMap)
+	{
+		if (Export.ObjectName.ToString() == AssetName)
+		{
+			if (!Export.ClassIndex.IsImport())
+			{
+				return false;
+			}
+			const int32 ClassIndex = Export.ClassIndex.ToImport();
+			const FObjectImport& AssetClass = LinkerLoad->ImportMap[ClassIndex];
+			return AssetClass.ObjectName == NAME_ObjectRedirector;
+		}
+	}
+	return false;
+}
+
+// Find the path of a redirector relative to it's content directory
+static FString GetContentRelativePath(FLinkerLoad* LinkerLoad)
+{
+	for (const FObjectImport& Import : LinkerLoad->ImportMap)
+    {
+        if (Import.ClassName != NAME_Package)
+        {
+        	continue;
+        }
+		
+        const FString PackageName = Import.ObjectName.ToString().RightChop(1);
+        if (PackageName.StartsWith(TEXT("Script/"), ESearchCase::CaseSensitive))
+        {
+        	continue;
+        }
+        
+        int32 ChopIndex;
+        if (PackageName.FindChar('/', ChopIndex))
+        {
+            return PackageName.RightChop(ChopIndex + 1);
+        }
+    }
+	return FString();
+}
+
+// Attempt to find the content directory and remove it from the path
+static FString GetContentRelativePath(const TSharedPtr<FChangelistFileData> &File)
+{
+	const FString BasePath = FPaths::GetBaseFilename(File->AssetDepotPath, /* bRemovePath */ false);
+	const FString ChopKey = TEXT("/Content/");
+	const int32 ChopIndex = BasePath.Find(ChopKey, ESearchCase::CaseSensitive, ESearchDir::FromEnd) + ChopKey.Len();
+	if (ChopIndex != INDEX_NONE)
+	{
+		return BasePath.RightChop(ChopIndex);
+	}
+	return FString();
+}
+
+// Identify renames or moves withing a CL and merge their redirected path into a single diff entry
+void SSourceControlReview::FixupRedirectors()
+{
+	// Search CL for redirectors and map them by their redirected path
+	TMap<FString, TWeakPtr<FChangelistFileData>> RedirectorsFound;
+	for (TSharedPtr<FChangelistFileData> ChangelistFile : ChangelistFiles)
+	{
+		if (FLinkerLoad* ReviewFileLoad = LoadPackageLinker(nullptr, ChangelistFile->ReviewFileTempPath, LOAD_ForDiff | LOAD_DisableCompileOnLoad | LOAD_DisableEngineVersionChecks))
+		{
+			if (IsObjectRedirector(ReviewFileLoad, ChangelistFile->AssetName))
+			{
+				FString Path = GetContentRelativePath(ReviewFileLoad);
+				if (!Path.IsEmpty())
+				{
+					RedirectorsFound.Add(Path, ChangelistFile.ToWeakPtr());
+				}
+			}
+		}
+	}
+
+	// Find files that match the redirected paths and merge the two entries into one
+	ChangelistFiles.RemoveAllSwap([&RedirectorsFound](const TSharedPtr<FChangelistFileData> &ChangelistFile)
+	{
+		if (ChangelistFile->FileSourceControlAction == ESourceControlAction::Add)
+		{
+			const FString Path = GetContentRelativePath(ChangelistFile);
+			if (const TWeakPtr<FChangelistFileData>* Found = RedirectorsFound.Find(Path))
+			{
+				if (const TSharedPtr<FChangelistFileData>& MergedRenameEntry = Found->Pin())
+				{
+					MergedRenameEntry->ReviewFileTempPath = ChangelistFile->ReviewFileTempPath;
+					MergedRenameEntry->PreviousAssetName = MergedRenameEntry->AssetName;
+					MergedRenameEntry->AssetName = ChangelistFile->AssetName;
+					MergedRenameEntry->RelativeFilePath = ChangelistFile->RelativeFilePath;
+					MergedRenameEntry->AssetFilePath = ChangelistFile->AssetFilePath;
+					MergedRenameEntry->PreviousFileRevisionNum = MergedRenameEntry->ReviewFileRevisionNum;
+					MergedRenameEntry->ReviewFileRevisionNum = ChangelistFile->ReviewFileRevisionNum;
+					return true;
+				}
+			}
+		}
+		return false;
+	});
+}
+
 SHeaderRow::FColumn::FArguments SSourceControlReview::HeaderColumn(FName HeaderName)
 {
 	FText ColumnLabel;
@@ -841,7 +908,12 @@ SHeaderRow::FColumn::FArguments SSourceControlReview::HeaderColumn(FName HeaderN
 FString SSourceControlReview::AsAssetPath(const FString& FullCLPath)
 {
 	const FString ProjectName = FString::Format(TEXT("/{0}/"), {FApp::GetProjectName()});
-	return UKismetSystemLibrary::GetProjectDirectory() / FullCLPath.RightChop(FullCLPath.Find(ProjectName) + ProjectName.Len());
+	const int32 Found = FullCLPath.Find(ProjectName);
+	if (Found == INDEX_NONE)
+	{
+		return FString();
+	}
+	return UKismetSystemLibrary::GetProjectDirectory() / FullCLPath.RightChop(Found + ProjectName.Len());
 }
 
 #undef LOCTEXT_NAMESPACE
diff --git a/Engine/Plugins/Editor/ChangelistReview/Source/Private/SSourceControlReview.h b/Engine/Plugins/Editor/ChangelistReview/Source/Private/SSourceControlReview.h
--- a/Engine/Plugins/Editor/ChangelistReview/Source/Private/SSourceControlReview.h
+++ b/Engine/Plugins/Editor/ChangelistReview/Source/Private/SSourceControlReview.h
@@ -39,31 +39,18 @@ namespace SourceControlReview
 
 	struct FChangelistFileData
 	{
-		FChangelistFileData()
-		{
-		}
-
-		FChangelistFileData(const FString& InAssetName, const FString& InReviewFilePkgName,
-							const FString& InReviewFileRevisionNum, const FString& InPreviousFilePkgName,
-							const FString& InPreviousFileRevisionNum)
-			: AssetName(InAssetName)
-			, ReviewFileName(InReviewFilePkgName)
-			, ReviewFileRevisionNum(InReviewFileRevisionNum)
-			, PreviousAssetName() // assume no rename by default
-			, PreviousFileName(InPreviousFilePkgName)
-			, PreviousFileRevisionNum(InPreviousFileRevisionNum)
-		{
-		}
+		FChangelistFileData() = default;
 
 		bool IsDataValidForEntry() const
 		{
-			return FileSourceControlAction != ESourceControlAction::Unset && !ReviewFileName.IsEmpty() &&
-				(!PreviousFileName.IsEmpty() || FileSourceControlAction == ESourceControlAction::Add || FileSourceControlAction == ESourceControlAction::Branch);
+			return FileSourceControlAction != ESourceControlAction::Unset && !ReviewFileTempPath.IsEmpty() &&
+				(!PreviousFileTempPath.IsEmpty() || FileSourceControlAction == ESourceControlAction::Add || FileSourceControlAction == ESourceControlAction::Branch);
 		}
 
 		FString AssetName;
 
-		FString ReviewFileName;
+		// path where the temporary asset was downloaded to
+		FPackagePath ReviewFileTempPath;
 
 		FString ReviewFileRevisionNum;
 
@@ -71,13 +58,16 @@ namespace SourceControlReview
 
 		FString PreviousAssetName;
 
-		FString PreviousFileName;
+		// path where the temporary asset was downloaded to
+		FPackagePath PreviousFileTempPath;
 
 		FString PreviousFileRevisionNum;
 
 		FString RelativeFilePath;
 
 		FString AssetFilePath;
+		
+		FString AssetDepotPath;
 
 		int32 ChangelistNum = INDEX_NONE;
 
@@ -191,6 +181,8 @@ private:
 	 * Removes CurrentChangelistInfo.SharedPath from the beginning of FullCLPath
 	 */
 	FString TrimSharedPath(FString FullCLPath) const;
+	
+	void FixupRedirectors();
 
 	static SHeaderRow::FColumn::FArguments HeaderColumn(FName HeaderName);
 
@@ -204,7 +196,6 @@ private:
 	uint32 FilesToLoad = 0;
 	uint32 FilesLoaded = 0;
 	TArray<TSharedPtr<FChangelistFileData>> ChangelistFiles;
-	TMap<FString, TWeakPtr<FChangelistFileData>> RedirectorsFound;
 	TSharedPtr<FGetChangelistDetails> GetChangelistDetailsCommand;
 	TArray<TSharedPtr<SourceControlReview::FChangelistLightInfo>> CLHistory;
 	bool bUncommittedChangelistNum = false;
diff --git a/Engine/Plugins/Editor/ChangelistReview/Source/Private/SSourceControlReviewEntry.cpp b/Engine/Plugins/Editor/ChangelistReview/Source/Private/SSourceControlReviewEntry.cpp
--- a/Engine/Plugins/Editor/ChangelistReview/Source/Private/SSourceControlReviewEntry.cpp
+++ b/Engine/Plugins/Editor/ChangelistReview/Source/Private/SSourceControlReviewEntry.cpp
@@ -210,7 +210,7 @@ bool SSourceControlReviewEntry::CanDiff() const
 void SSourceControlReviewEntry::TryBindUAssetDiff()
 {
 	UObject* ReviewAsset = nullptr;
-	if (UPackage* ReviewFilePkg = LoadPackage(nullptr, *ChangelistFileData.ReviewFileName, LOAD_ForDiff | LOAD_DisableCompileOnLoad | LOAD_DisableEngineVersionChecks))
+	if (UPackage* ReviewFilePkg = LoadPackage(nullptr, ChangelistFileData.ReviewFileTempPath, LOAD_ForDiff | LOAD_DisableCompileOnLoad | LOAD_DisableEngineVersionChecks))
 	{
 		ReviewAsset = FindObject<UObject>(ReviewFilePkg, *ChangelistFileData.AssetName);
 		if(ReviewAsset && ReviewAsset->IsA<UObjectRedirector>())
@@ -219,9 +219,9 @@ void SSourceControlReviewEntry::TryBindUAssetDiff()
 		}
 	}
 	UObject* PreviousAsset = nullptr;
-	if (!ChangelistFileData.PreviousFileName.IsEmpty())
+	if (!ChangelistFileData.PreviousFileTempPath.IsEmpty())
 	{
-		if (UPackage* PreviousFilePkg = LoadPackage(nullptr, *ChangelistFileData.PreviousFileName, LOAD_ForDiff | LOAD_DisableCompileOnLoad | LOAD_DisableEngineVersionChecks))
+		if (UPackage* PreviousFilePkg = LoadPackage(nullptr, ChangelistFileData.PreviousFileTempPath, LOAD_ForDiff | LOAD_DisableCompileOnLoad | LOAD_DisableEngineVersionChecks))
 		{
 			if (ChangelistFileData.PreviousAssetName.IsEmpty())
 			{
@@ -291,7 +291,7 @@ void SSourceControlReviewEntry::TryBindTextDiff()
 	{
 		const FString& DiffCommand = GetDefault<UEditorLoadingSavingSettings>()->TextDiffToolPath.FilePath;
 		const FAssetToolsModule& AssetToolsModule = FAssetToolsModule::GetModule();
-		AssetToolsModule.Get().CreateDiffProcess(DiffCommand, ChangelistFileData.PreviousFileName, ChangelistFileData.ReviewFileName);
+		AssetToolsModule.Get().CreateDiffProcess(DiffCommand, ChangelistFileData.PreviousFileTempPath.GetLocalFullPath(), ChangelistFileData.ReviewFileTempPath.GetLocalFullPath());
 	});
 }
 
