[BugFix] The shared path in the changelist review tool is being split in the middle of filenames instead of keeping full directory names together

diff --git a/Engine/Plugins/Editor/ChangelistReview/Source/Private/SSourceControlReview.cpp b/Engine/Plugins/Editor/ChangelistReview/Source/Private/SSourceControlReview.cpp
--- a/Engine/Plugins/Editor/ChangelistReview/Source/Private/SSourceControlReview.cpp
+++ b/Engine/Plugins/Editor/ChangelistReview/Source/Private/SSourceControlReview.cpp
@@ -26,6 +26,7 @@
 #include "Internationalization/Regex.h"
 #include "Widgets/Input/SComboBox.h"
 #include "Misc/ConfigCacheIni.h"
+#include "UObject/LinkerLoad.h"
 
 #define LOCTEXT_NAMESPACE "SourceControlReview"
 
@@ -179,7 +180,10 @@ void SSourceControlReview::Construct(const FArguments& InArgs)
 				.HAlign(HAlign_Fill)
 				[
 					SNew(SEditableTextBox)
-					.Text_Lambda([this]{return CurrentChangelistInfo.SharedPath;})
+					.Text_Lambda([this]
+					{
+						return CurrentChangelistInfo.SharedPath.IsEmpty() ? FText::GetEmpty() : FText::Format(LOCTEXT("CLPath", "{0}/..."), CurrentChangelistInfo.SharedPath);
+					})
 					.IsReadOnly(true)
 					.Style(&InfoWidgetStyle)
 				]
@@ -729,25 +733,27 @@ static FString GetSharedBranchPath(const TMap<FString, FString>& InChangelistRec
 	uint32  RecordFileIndex = 1;
 	//The p4 records is the map a file key starts with "depotFile" and is followed by file index 
 	FString RecordFileMapKey = ReviewHelpers::FileDepotKey + LexToString(RecordFileIndex);
+	int32 TrimIndex = SharedBranchPath.Len();
 	while (const FString* Found = InChangelistRecord.Find(RecordFileMapKey))
 	{
-		// find starting from the left, find the portion that's shared between both strings
-		int32 TrimIndex = 0;
-		const int32 MaxTrimIndex = FMath::Min(Found->Len(), SharedBranchPath.Len());
-		for(; TrimIndex < MaxTrimIndex; ++TrimIndex)
+		// starting from the left, find the portion that's shared between both strings
+		int32 CurrentIndex = 0;
+		TrimIndex = FMath::Min(Found->Len(), TrimIndex);
+		for(; CurrentIndex < TrimIndex; ++CurrentIndex)
 		{
-			if (SharedBranchPath[TrimIndex] != (*Found)[TrimIndex])
+			if (SharedBranchPath[CurrentIndex] != (*Found)[CurrentIndex])
 			{
 				break;
 			}
 		}
-		// shorten SharedBranchPath to only contain the shared portion
-		SharedBranchPath = SharedBranchPath.Left(TrimIndex);
+		TrimIndex = CurrentIndex;
 
 		// increment to next file path
 		RecordFileMapKey = ReviewHelpers::FileDepotKey + LexToString(++RecordFileIndex);
 	}
-	return SharedBranchPath;
+	// if the shared string is in the middle of a file/directory name, backtrack to the last seen '/'
+	TrimIndex = SharedBranchPath.Find(TEXT("/"), ESearchCase::IgnoreCase, ESearchDir::FromEnd, TrimIndex);
+	return TrimIndex == INDEX_NONE ? FString() : SharedBranchPath.Left(TrimIndex);
 }
 
 void SSourceControlReview::SetChangelistInfo(const TMap<FString, FString>& InChangelistRecord)
