[BugFix] Rewrote GetContentRelativePath functions to use regex and support paths with nested Content directories

diff --git a/Engine/Plugins/Editor/ChangelistReview/Source/Private/SSourceControlReview.cpp b/Engine/Plugins/Editor/ChangelistReview/Source/Private/SSourceControlReview.cpp
--- a/Engine/Plugins/Editor/ChangelistReview/Source/Private/SSourceControlReview.cpp
+++ b/Engine/Plugins/Editor/ChangelistReview/Source/Private/SSourceControlReview.cpp
@@ -805,17 +805,22 @@ static FString GetContentRelativePath(FLinkerLoad* LinkerLoad)
         	continue;
         }
 		
-        const FString PackageName = Import.ObjectName.ToString().RightChop(1);
-        if (PackageName.StartsWith(TEXT("Script/"), ESearchCase::CaseSensitive))
+        const FString PackageName = Import.ObjectName.ToString();
+		
+		// ignore internal packages
+        if (PackageName.StartsWith(TEXT("/Script/"), ESearchCase::CaseSensitive))
         {
         	continue;
         }
-        
-        int32 ChopIndex;
-        if (PackageName.FindChar('/', ChopIndex))
-        {
-            return PackageName.RightChop(ChopIndex + 1);
-        }
+
+		// Matches after the last content directory in the path
+		// if content directory isn't found, assume this is a plugin and match after the first directory
+		const FRegexPattern Pattern(TEXT(R"((?:(?:^.*/Content/)|(?:^/[^/]*/))(.*))"));
+		FRegexMatcher Regex(Pattern, PackageName);
+		if (Regex.FindNext())
+		{
+			return Regex.GetCaptureGroup(1);
+		}
     }
 	return FString();
 }
@@ -823,14 +828,11 @@ static FString GetContentRelativePath(FLinkerLoad* LinkerLoad)
 // Attempt to find the content directory and remove it from the path
 static FString GetContentRelativePath(const TSharedPtr<FChangelistFileData> &File)
 {
+	// Matches after the last content directory in the path
 	const FString BasePath = FPaths::GetBaseFilename(File->AssetDepotPath, /* bRemovePath */ false);
-	const FString ChopKey = TEXT("/Content/");
-	const int32 ChopIndex = BasePath.Find(ChopKey, ESearchCase::CaseSensitive, ESearchDir::FromEnd) + ChopKey.Len();
-	if (ChopIndex != INDEX_NONE)
-	{
-		return BasePath.RightChop(ChopIndex);
-	}
-	return FString();
+	const FRegexPattern Pattern(TEXT(R"((?:^.*/Content/)(.*))"));
+	FRegexMatcher Regex(Pattern, BasePath);
+	return Regex.FindNext()? Regex.GetCaptureGroup(1) : FString();
 }
 
 // Identify renames or moves withing a CL and merge their redirected path into a single diff entry
