[BugFix] Disabled ReviewChangelist tool when source control is inactive
also made SPerforceSourceControlSettings::GetPassword() fail silently when on a non-slate thread rather than asserting

diff --git a/Engine/Plugins/Developer/PerforceSourceControl/Source/PerforceSourceControl/Private/SPerforceSourceControlSettings.cpp b/Engine/Plugins/Developer/PerforceSourceControl/Source/PerforceSourceControl/Private/SPerforceSourceControlSettings.cpp
--- a/Engine/Plugins/Developer/PerforceSourceControl/Source/PerforceSourceControl/Private/SPerforceSourceControlSettings.cpp
+++ b/Engine/Plugins/Developer/PerforceSourceControl/Source/PerforceSourceControl/Private/SPerforceSourceControlSettings.cpp
@@ -218,7 +218,8 @@ void SPerforceSourceControlSettings::Construct(const FArguments& InArgs, FPerfor
 
 FString SPerforceSourceControlSettings::GetPassword()
 {
-	if(PasswordTextBox.IsValid())
+	const bool bIsValidThread = IsInGameThread() || IsInSlateThread();
+	if(PasswordTextBox.IsValid() && bIsValidThread)
 	{
 		return PasswordTextBox.Pin()->GetText().ToString();
 	}
diff --git a/Engine/Plugins/Developer/PerforceSourceControl/Source/PerforceSourceControl/Private/SPerforceSourceControlSettings.h b/Engine/Plugins/Developer/PerforceSourceControl/Source/PerforceSourceControl/Private/SPerforceSourceControlSettings.h
--- a/Engine/Plugins/Developer/PerforceSourceControl/Source/PerforceSourceControl/Private/SPerforceSourceControlSettings.h
+++ b/Engine/Plugins/Developer/PerforceSourceControl/Source/PerforceSourceControl/Private/SPerforceSourceControlSettings.h
@@ -43,7 +43,10 @@ public:
 
 	void Construct(const FArguments& InArgs, FPerforceSourceControlProvider* InSCCProvider);
 
-	/** Get the currently entered password */
+	/** Get the currently entered password
+	 *
+	 *  Note: This can only be called from a game or slate thread. All other threads will return an empty string.
+	 */
 	static FString GetPassword();
 
 	virtual void Tick( const FGeometry& AllottedGeometry, const double InCurrentTime, const float InDeltaTime );
diff --git a/Engine/Plugins/Editor/ChangelistReview/Source/Private/SSourceControlReview.cpp b/Engine/Plugins/Editor/ChangelistReview/Source/Private/SSourceControlReview.cpp
--- a/Engine/Plugins/Editor/ChangelistReview/Source/Private/SSourceControlReview.cpp
+++ b/Engine/Plugins/Editor/ChangelistReview/Source/Private/SSourceControlReview.cpp
@@ -23,6 +23,7 @@
 #include "Internationalization/Regex.h"
 #include "Widgets/Input/SComboBox.h"
 #include "Misc/ConfigCacheIni.h"
+#include "SourceControl/Private/SourceControlModule.h"
 #include "UObject/LinkerLoad.h"
 
 #define LOCTEXT_NAMESPACE "SourceControlReview"
@@ -135,7 +136,9 @@ void SSourceControlReview::Construct(const FArguments& InArgs)
 					[
 						SNew(SButton)
 						.HAlign(HAlign_Center)
+						.IsEnabled(this, &SSourceControlReview::IsSourceControlActive)
 						.OnClicked(this, &SSourceControlReview::OnLoadChangelistClicked)
+						.ToolTipText(this, &SSourceControlReview::LoadChangelistTooltip)
 						[
 							SNew(STextBlock)
 							.Text(LOCTEXT("LoadChangelistText", "Load"))
@@ -446,6 +449,20 @@ FReply SSourceControlReview::OnLoadChangelistClicked()
 	return FReply::Handled();
 }
 
+bool SSourceControlReview::IsSourceControlActive() const
+{
+	return FSourceControlModule::Get().IsEnabled();
+}
+
+FText SSourceControlReview::LoadChangelistTooltip() const
+{
+	if (IsSourceControlActive())
+	{
+		return LOCTEXT("LoadChangelistTooltip", "Load Changelist");
+	}
+	return LOCTEXT("LoadChangelistInactive", "Enable Revision Control to load changelist");
+}
+
 void SSourceControlReview::OnChangelistNumChanged(const FText& Text)
 {
 	const FString& Data = Text.ToString();
diff --git a/Engine/Plugins/Editor/ChangelistReview/Source/Private/SSourceControlReview.h b/Engine/Plugins/Editor/ChangelistReview/Source/Private/SSourceControlReview.h
--- a/Engine/Plugins/Editor/ChangelistReview/Source/Private/SSourceControlReview.h
+++ b/Engine/Plugins/Editor/ChangelistReview/Source/Private/SSourceControlReview.h
@@ -132,6 +132,8 @@ public:
 private:
 	void OnChangelistLoadComplete(const FSourceControlOperationRef& InOperation, ECommandResult::Type InResult, FString Changelist);
 	FReply OnLoadChangelistClicked();
+	bool IsSourceControlActive() const;
+	FText LoadChangelistTooltip() const;
 	void OnChangelistNumChanged(const FText& Text);
 	void OnChangelistNumCommitted(const FText& Text, ETextCommit::Type CommitMethod);
 
