[Bugfix] Changelist History dropdown in Changelist Review Tool removes recent history if an invalid CL is loaded multiple times (causes crash if history is empty)
#lockout julien.marchand

diff --git a/Engine/Plugins/Editor/ChangelistReview/Source/Private/SSourceControlReview.cpp b/Engine/Plugins/Editor/ChangelistReview/Source/Private/SSourceControlReview.cpp
--- a/Engine/Plugins/Editor/ChangelistReview/Source/Private/SSourceControlReview.cpp
+++ b/Engine/Plugins/Editor/ChangelistReview/Source/Private/SSourceControlReview.cpp
@@ -303,15 +303,6 @@ void SSourceControlReview::LoadChangelist(const FString& Changelist)
 		return;
 	}
 
-	for(int32 I = 1; I < CLHistory.Num(); ++I)
-	{
-		if (CLHistory[I]->Number.EqualTo(CLHistory[0]->Number))
-		{
-			CLHistory.RemoveAt(I);
-			break;
-		}
-	}
-	bUncommittedChangelistNum = false;
 	ChangelistFiles.Empty();
 
 	//This command runs p4 -describe (or similar for other version controls) to retrieve changelist record information
@@ -324,6 +315,43 @@ void SSourceControlReview::LoadChangelist(const FString& Changelist)
 		FSourceControlOperationComplete::CreateRaw(this, &SSourceControlReview::OnChangelistLoadComplete, Changelist));
 }
 
+void SSourceControlReview::CommitTempChangelistNumToHistory()
+{
+	if (bUncommittedChangelistNum)
+	{
+		// bUncommittedChangelistNum should never be true when the history is empty
+		check(!CLHistory.IsEmpty())
+		
+		// temp cl nums are missing an author and description. we have that now so let's provide it.
+		// annoyingly the SComboBox doesn't visually update unless the shared pointer is different so we have to reallocate it :(
+		CLHistory[0] = MakeShared<FChangelistLightInfo>(CLHistory[0]->Number, CurrentChangelistInfo.Author, CurrentChangelistInfo.Description);
+
+		// make sure there's only one copy of this CL in the history.
+		for(int32 I = 1; I < CLHistory.Num(); ++I)
+		{
+			if (CLHistory[I]->Number.EqualTo(CLHistory[0]->Number))
+			{
+				CLHistory.RemoveAt(I);
+				break; // there should only be a max of 1 duplicate so we can break early
+			}
+		}
+		ChangelistNumComboBox->RefreshOptions();
+		bUncommittedChangelistNum = false;
+	}
+}
+
+void SSourceControlReview::RemoveUncommittedChangelistNumFromHistory()
+{
+	if (bUncommittedChangelistNum)
+	{
+		// bUncommittedChangelistNum should never be true when the history is empty
+		check(!CLHistory.IsEmpty())
+		
+		CLHistory.RemoveAt(0); 
+		bUncommittedChangelistNum = false;
+	}
+}
+
 void SSourceControlReview::OnChangelistLoadComplete(const FSourceControlOperationRef& InOperation, ECommandResult::Type InResult, FString Changelist)
 {
 	// this command is cancelled when this widget is destroyed. Exit immediately to avoid touching invalid data
@@ -337,7 +365,7 @@ void SSourceControlReview::OnChangelistLoadComplete(const FSourceControlOperatio
 	
 	if (!IsChangelistRecordValid(Record))
 	{
-		CLHistory.RemoveAt(0); // don't save history of invalid changelists
+		RemoveUncommittedChangelistNumFromHistory();
 		SetLoading(false);
 		return;
 	}
@@ -362,12 +390,7 @@ void SSourceControlReview::OnChangelistLoadComplete(const FSourceControlOperatio
 	FString RecordActionMapKey = ReviewHelpers::FileActionKey + RecordFileIndexStr;
 
 	SetChangelistInfo(ChangelistRecord);
-	if (CLHistory[0]->Author.IsEmpty())
-	{
-		// annoyingly the options list doesn't visually update unless the shared pointer is different
-		CLHistory[0] = MakeShared<FChangelistLightInfo>(CLHistory[0]->Number, CurrentChangelistInfo.Author, CurrentChangelistInfo.Description);
-	}
-	ChangelistNumComboBox->RefreshOptions();
+	CommitTempChangelistNumToHistory();
 	SaveCLHistory();
 	
 	//the loop checks if we have a valid record "depotFile(Index)" in the records to add a file entry
diff --git a/Engine/Plugins/Editor/ChangelistReview/Source/Private/SSourceControlReview.h b/Engine/Plugins/Editor/ChangelistReview/Source/Private/SSourceControlReview.h
--- a/Engine/Plugins/Editor/ChangelistReview/Source/Private/SSourceControlReview.h
+++ b/Engine/Plugins/Editor/ChangelistReview/Source/Private/SSourceControlReview.h
@@ -206,6 +206,12 @@ private:
 	 */
 	static FString AsAssetPath(const FString& FullCLPath);
 
+	// helper function that sets bUncommittedChangelistNum to false and upgrades CLHistory[0] from a entry to an official one
+	void CommitTempChangelistNumToHistory();
+	
+	// helper function that sets bUncommittedChangelistNum to false and removes CLHistory[0] if it's uncommitted
+	void RemoveUncommittedChangelistNumFromHistory();
+
 	// used for asynchronous changelist loading
 	bool bChangelistLoading = false;
 	uint32 FilesToLoad = 0;
