[BugFix] Changelist Review tool can be used while not connected to Source Control, resulting in infinite loading.
 - users could still bypass the load button being disabled by pressing enter while the CL input box is selected. Fixed by additionally disabling the CL input box
#lockout julien.marchand

diff --git a/Engine/Plugins/Editor/ChangelistReview/Source/Private/SSourceControlReview.cpp b/Engine/Plugins/Editor/ChangelistReview/Source/Private/SSourceControlReview.cpp
--- a/Engine/Plugins/Editor/ChangelistReview/Source/Private/SSourceControlReview.cpp
+++ b/Engine/Plugins/Editor/ChangelistReview/Source/Private/SSourceControlReview.cpp
@@ -117,6 +117,7 @@ void SSourceControlReview::Construct(const FArguments& InArgs)
 					[
 						SAssignNew(ChangelistNumComboBox, SComboBox<TSharedPtr<FChangelistLightInfo>>)
 						.OptionsSource(&CLHistory)
+                     	.IsEnabled(this, &SSourceControlReview::IsSourceControlActive)
 						.OnGenerateWidget(this, &SSourceControlReview::MakeCLComboOption)
 						.OnSelectionChanged(this, &SSourceControlReview::OnCLComboSelection)
 						.ContentPadding(0.f)
@@ -126,6 +127,7 @@ void SSourceControlReview::Construct(const FArguments& InArgs)
                      		.Font(FStyleFonts::Get().Normal)
                      		.MinDesiredWidth(55.f)
                      		.Justification(ETextJustify::Center)
+                     		.IsEnabled(this, &SSourceControlReview::IsSourceControlActive)
                      		.OnTextCommitted(this, &SSourceControlReview::OnChangelistNumCommitted)
                      		.OnTextChanged(this, &SSourceControlReview::OnChangelistNumChanged)
 						]
