Remove unnecessary private include

#rnx

diff --git a/Engine/Plugins/Editor/ChangelistReview/Source/Private/SSourceControlReview.cpp b/Engine/Plugins/Editor/ChangelistReview/Source/Private/SSourceControlReview.cpp
--- a/Engine/Plugins/Editor/ChangelistReview/Source/Private/SSourceControlReview.cpp
+++ b/Engine/Plugins/Editor/ChangelistReview/Source/Private/SSourceControlReview.cpp
@@ -23,7 +23,6 @@
 #include "Internationalization/Regex.h"
 #include "Widgets/Input/SComboBox.h"
 #include "Misc/ConfigCacheIni.h"
-#include "SourceControl/Private/SourceControlModule.h"
 #include "UObject/LinkerLoad.h"
 
 #define LOCTEXT_NAMESPACE "SourceControlReview"
@@ -478,7 +477,7 @@ FReply SSourceControlReview::OnLoadChangelistClicked()
 
 bool SSourceControlReview::IsSourceControlActive() const
 {
-	return FSourceControlModule::Get().IsEnabled();
+	return FModuleManager::LoadModuleChecked<ISourceControlModule>("SourceControl").IsEnabled();
 }
 
 FText SSourceControlReview::LoadChangelistTooltip() const
