[Feature] Review Comments API
- setup for making the api calls neccesary for supporting comments in the changelist review tool

diff --git a/Engine/Plugins/Editor/ChangelistReview/Source/ChangelistReview.Build.cs b/Engine/Plugins/Editor/ChangelistReview/Source/ChangelistReview.Build.cs
--- a/Engine/Plugins/Editor/ChangelistReview/Source/ChangelistReview.Build.cs
+++ b/Engine/Plugins/Editor/ChangelistReview/Source/ChangelistReview.Build.cs
@@ -21,11 +21,22 @@ public class ChangelistReview : ModuleRules
 				"Slate",
 				"EditorFramework",
 				"UnrealEd",
+				"Kismet",
 				"EditorStyle",
 				"ToolMenus",
 				"SourceControl",
-				"InputCore"
+				"InputCore",
+				"HTTP",
+				"Json"
 			}
 		);
+
+		AddEngineThirdPartyPrivateStaticDependencies(Target, "Perforce");
+
+		if (Target.Platform == UnrealTargetPlatform.Win64 || Target.Platform == UnrealTargetPlatform.Mac)
+		{
+			AddEngineThirdPartyPrivateStaticDependencies(Target, "OpenSSL");
+			AddEngineThirdPartyPrivateStaticDependencies(Target, "zlib");
+		}
 	}
 }
\ No newline at end of file
diff --git a/Engine/Plugins/Editor/ChangelistReview/Source/Private/FSwarmCommentsAPI.cpp b/Engine/Plugins/Editor/ChangelistReview/Source/Private/FSwarmCommentsAPI.cpp
new file mode 100644
--- /dev/null
+++ b/Engine/Plugins/Editor/ChangelistReview/Source/Private/FSwarmCommentsAPI.cpp
@@ -0,0 +1,532 @@
+// Copyright Epic Games, Inc. All Rights Reserved.
+#include "FSwarmCommentsAPI.h"
+
+#include "HttpModule.h"
+#include "ISourceControlModule.h"
+#include "Interfaces/IHttpResponse.h"
+#include "Misc/Base64.h"
+#include "Misc/FileHelper.h"
+#include "Misc/Paths.h"
+#include "Serialization/JsonSerializer.h"
+
+THIRD_PARTY_INCLUDES_START
+// used to retrieve the swarm url from p4
+#include <p4/clientapi.h>
+THIRD_PARTY_INCLUDES_END
+
+FSwarmCommentsAPI::FSwarmCommentsAPI()
+	: AuthTicket(RetrieveAuthorizationTicket())
+	, SwarmURL(RetrieveSwarmURL())
+{
+}
+
+FString FSwarmCommentsAPI::GetUsername() const
+{
+	return AuthTicket.Username;
+}
+
+void FSwarmCommentsAPI::GetComments(const FReviewTopic& Topic, const OnGetCommentsComplete& OnComplete) const
+{
+    const TSharedRef<IHttpRequest> HttpRequest = FHttpModule::Get().CreateRequest();
+	HttpRequest->SetHeader(TEXT("Content-Type"), TEXT("application/json; charset=utf-8"));
+	HttpRequest->SetHeader(TEXT("Authorization"), AuthTicket);
+	
+    HttpRequest->SetURL(FString::Format(TEXT("{0}?topic={1}&max={2}"),
+    	{
+    		CommentsURL(),
+    		Topic.ToString(), // filter to comments in a specific review
+    		5000 // limit to 5k comments. This should obviously be more than enough and loads at a reasonable rate
+    	}));
+    HttpRequest->SetVerb(TEXT("GET"));
+	
+	HttpRequest->OnProcessRequestComplete().BindLambda([OnComplete, this](FHttpRequestPtr Request, FHttpResponsePtr Response, bool bConnectedSuccessfully)
+	{
+		if (!bConnectedSuccessfully)
+		{
+			OnComplete.ExecuteIfBound({}, TEXT("Connection Failed"));
+			return;
+		}
+		const TSharedRef<TJsonReader<>> JsonReader = TJsonReaderFactory<>::Create(Response->GetContentAsString());
+		TSharedPtr<FJsonObject> JsonObject;
+		if (FJsonSerializer::Deserialize(JsonReader, JsonObject))
+		{
+			if (JsonObject->HasField("error"))
+			{
+				OnComplete.ExecuteIfBound({}, JsonObject->GetStringField("error"));
+				return;
+			}
+			if (JsonObject->HasField("comments"))
+			{
+				TArray<FReviewComment> Comments;
+				const TArray<TSharedPtr<FJsonValue>> CommentsJson = JsonObject->GetArrayField("comments");
+ 				for (const TSharedPtr<FJsonValue>& CommentJson : CommentsJson)
+				{
+ 					FReviewComment Comment = FReviewComment::FromJson(CommentJson->AsObject());
+ 					TakeMetadataFromBody(Comment);
+					Comments.Add(Comment);
+				}
+				OnComplete.ExecuteIfBound(MoveTemp(Comments), {});
+				return;
+			}
+		}
+		OnComplete.ExecuteIfBound({}, TEXT("Malformed Response"));
+	});
+		
+	
+	HttpRequest->ProcessRequest();
+}
+
+void FSwarmCommentsAPI::PostComment(FReviewComment& Comment, const OnPostCommentComplete& OnComplete) const
+{
+	const TSharedRef<IHttpRequest> HttpRequest = FHttpModule::Get().CreateRequest();
+	HttpRequest->SetHeader(TEXT("Content-Type"), TEXT("application/json; charset=utf-8"));
+	HttpRequest->SetHeader(TEXT("Authorization"), AuthTicket);
+	HttpRequest->SetURL(CommentsURL());
+	HttpRequest->SetVerb(TEXT("POST"));
+
+	Comment.User = AuthTicket.Username;
+	FReviewComment CleanedComment = Comment;
+	PutMetadataInBody(CleanedComment);
+	HttpRequest->SetContentAsString(CleanedComment.ToJson());
+	
+	HttpRequest->OnProcessRequestComplete().BindLambda([OnComplete, this](FHttpRequestPtr Request, FHttpResponsePtr Response, bool bConnectedSuccessfully)
+	{
+		if (!bConnectedSuccessfully)
+		{
+			OnComplete.ExecuteIfBound({}, TEXT("Connection Failed"));
+			return;
+		}
+		const TSharedRef<TJsonReader<>> JsonReader = TJsonReaderFactory<>::Create(Response->GetContentAsString());
+		TSharedPtr<FJsonObject> JsonObject;
+		if (FJsonSerializer::Deserialize(JsonReader, JsonObject))
+		{
+			if (JsonObject->HasField("error"))
+			{
+				OnComplete.ExecuteIfBound({}, JsonObject->GetStringField("error"));
+				return;
+			}
+			if (JsonObject->HasField("comment"))
+			{
+				FReviewComment Comment = FReviewComment::FromJson(JsonObject->GetObjectField("comment"));
+				TakeMetadataFromBody(Comment);
+				OnComplete.ExecuteIfBound(Comment, {});
+				return;
+			}
+		}
+		OnComplete.ExecuteIfBound({}, TEXT("Malformed Response"));
+	});
+	
+	HttpRequest->ProcessRequest();
+}
+
+void FSwarmCommentsAPI::EditComment(const FReviewComment& Comment, const OnEditCommentComplete& OnComplete) const
+{
+	const TSharedRef<IHttpRequest> HttpRequest = FHttpModule::Get().CreateRequest();
+	HttpRequest->SetHeader(TEXT("Content-Type"), TEXT("application/json; charset=utf-8"));
+	HttpRequest->SetHeader(TEXT("Authorization"), AuthTicket);
+	const FString CommentID = FString::FromInt(Comment.CommentID.GetValue());
+	HttpRequest->SetURL(CommentsURL() / CommentID);
+	HttpRequest->SetVerb(TEXT("PATCH"));
+	
+	FReviewComment CleanedComment = Comment;
+	PutMetadataInBody(CleanedComment);
+	HttpRequest->SetContentAsString(CleanedComment.ToJson());
+	HttpRequest->OnProcessRequestComplete().BindLambda([OnComplete, this](FHttpRequestPtr Request, FHttpResponsePtr Response, bool bConnectedSuccessfully)
+	{
+		if (!bConnectedSuccessfully)
+		{
+			OnComplete.ExecuteIfBound({}, TEXT("Connection Failed"));
+			return;
+		}
+		const TSharedRef<TJsonReader<>> JsonReader = TJsonReaderFactory<>::Create(Response->GetContentAsString());
+		TSharedPtr<FJsonObject> JsonObject;
+		if (FJsonSerializer::Deserialize(JsonReader, JsonObject))
+		{
+			if (JsonObject->HasField("error"))
+			{
+				OnComplete.ExecuteIfBound({}, JsonObject->GetStringField("error"));
+				return;
+			}
+			if (JsonObject->HasField("comment"))
+			{
+				OnComplete.ExecuteIfBound(FReviewComment::FromJson(JsonObject->GetObjectField("comment")), {});
+				return;
+			}
+		}
+		OnComplete.ExecuteIfBound({}, TEXT("Malformed Response"));
+	});
+	
+	HttpRequest->ProcessRequest();
+}
+
+void FSwarmCommentsAPI::GetReviewTopicForCL(const FString& ChangelistNum,
+	const OnGetReviewTopicForCLComplete& OnComplete) const
+{
+	const TSharedRef<IHttpRequest> HttpRequest = FHttpModule::Get().CreateRequest();
+	HttpRequest->SetHeader(TEXT("Content-Type"), TEXT("application/json; charset=utf-8"));
+	HttpRequest->SetHeader(TEXT("Authorization"), AuthTicket);
+	
+	HttpRequest->SetURL(FString::Format(TEXT("{0}?change={1}&max={2}"),
+		{
+			ReviewsURL(),
+			ChangelistNum, // get the review for a specific CL
+			1 // we only want a single review.
+		}));
+	HttpRequest->SetVerb(TEXT("GET"));
+	
+	HttpRequest->OnProcessRequestComplete().BindLambda([OnComplete, ChangelistNum, this](FHttpRequestPtr Request, FHttpResponsePtr Response, bool bConnectedSuccessfully)
+	{
+		if (!bConnectedSuccessfully)
+		{
+			OnComplete.ExecuteIfBound({}, TEXT("Connection Failed"));
+			return;
+		}
+		const TSharedRef<TJsonReader<>> JsonReader = TJsonReaderFactory<>::Create(Response->GetContentAsString());
+		TSharedPtr<FJsonObject> JsonObject;
+		if (FJsonSerializer::Deserialize(JsonReader, JsonObject))
+		{
+			if (JsonObject->HasField("error"))
+			{
+				OnComplete.ExecuteIfBound({}, JsonObject->GetStringField("error"));
+				return;
+			}
+			if (JsonObject->HasField("reviews"))
+			{
+				const TArray<TSharedPtr<FJsonValue>> ReviewsJson = JsonObject->GetArrayField("reviews");
+				if (ReviewsJson.IsEmpty())
+				{
+					CreateReviewTopicForCL(ChangelistNum, OnComplete);
+					return;
+				}
+				
+				const TSharedPtr<FJsonObject> Review = ReviewsJson[0]->AsObject();
+				
+				OnComplete.ExecuteIfBound(FReviewTopic{
+					FString::FromInt(Review->GetIntegerField(TEXT("id"))),
+					EReviewTopicType::Review
+				}, {});
+				return;
+			}
+		}
+		
+		OnComplete.ExecuteIfBound({}, TEXT("Malformed Response"));
+	});
+	
+	HttpRequest->ProcessRequest();
+}
+
+void FSwarmCommentsAPI::CreateReviewTopicForCL(const FString& ChangelistNum,
+	const OnGetReviewTopicForCLComplete& OnComplete) const
+{
+	
+	const TSharedRef<IHttpRequest> HttpRequest = FHttpModule::Get().CreateRequest();
+	HttpRequest->SetHeader(TEXT("Content-Type"), TEXT("application/json; charset=utf-8"));
+	HttpRequest->SetHeader(TEXT("Authorization"), AuthTicket);
+	
+	HttpRequest->SetURL(ReviewsURL());
+	HttpRequest->SetContentAsString(FString::Format(TEXT("{\"change\":{0}}"), {ChangelistNum}));
+	HttpRequest->SetVerb(TEXT("POST"));
+	
+	HttpRequest->OnProcessRequestComplete().BindLambda([OnComplete, ChangelistNum, this](FHttpRequestPtr Request, FHttpResponsePtr Response, bool bConnectedSuccessfully)
+	{
+		if (!bConnectedSuccessfully)
+		{
+			OnComplete.ExecuteIfBound({}, TEXT("Connection Failed"));
+			return;
+		}
+		const TSharedRef<TJsonReader<>> JsonReader = TJsonReaderFactory<>::Create(Response->GetContentAsString());
+		TSharedPtr<FJsonObject> JsonObject;
+		if (FJsonSerializer::Deserialize(JsonReader, JsonObject))
+		{
+			if (JsonObject->HasField("error"))
+			{
+				OnComplete.ExecuteIfBound({}, JsonObject->GetStringField("error"));
+				return;
+			}
+			if (JsonObject->HasField("review"))
+			{
+				const TSharedPtr<FJsonObject> ReviewJson = JsonObject->GetObjectField("review");
+				
+				OnComplete.ExecuteIfBound(FReviewTopic{
+					FString::FromInt(ReviewJson->GetIntegerField(TEXT("id"))),
+					EReviewTopicType::Review
+				}, {});
+				return;
+			}
+		}
+		
+		OnComplete.ExecuteIfBound({}, TEXT("Malformed Response"));
+	});
+	
+	HttpRequest->ProcessRequest();
+}
+
+FString FSwarmCommentsAPI::CommentsURL() const
+{
+	return SwarmURL.IsEmpty() ? FString{} : SwarmURL / TEXT("api/v9/comments");
+}
+
+FString FSwarmCommentsAPI::ReviewsURL() const
+{
+	return SwarmURL.IsEmpty() ? FString{} : SwarmURL / TEXT("api/v9/reviews");
+}
+
+
+static TMap<FString, FString> ParseReviewCommentMetadata(const FString& Comment)
+{
+	TMap<FString, FString> Result;
+	FString Body, Metadata;
+	if (!Comment.Split(TEXT("[metadata]"), &Body, &Metadata))
+	{
+		return {{TEXT("Body"), Comment.TrimEnd()}};
+	}
+	
+	Result.Add(TEXT("Body"), Body.TrimEnd());
+	
+	TArray<FString> MetaDataLines;
+	Metadata.ParseIntoArrayLines(MetaDataLines);
+	for (const FString& Line : MetaDataLines)
+	{
+		const int32 KeyEnd = Line.Find(TEXT("="), ESearchCase::CaseSensitive, ESearchDir::FromStart);
+		if (KeyEnd == INDEX_NONE)
+		{
+			continue;
+		}
+		FString Key = Line.Left(KeyEnd);
+		
+		const int32 ValBegin = KeyEnd + 1; // skip '=' char
+		const int32 ValLen = Line.Len() - ValBegin; // extend to end of line
+		const FString Val = Line.Mid(ValBegin, ValLen);
+		Result.Add(Key, Val);
+	}
+	return Result;
+}
+
+static TMap<FString, FString> ParseReviewCommentMetadata(const FReviewComment& Comment)
+{
+	return ParseReviewCommentMetadata(Comment.Body.Get({}));
+}
+
+static void SetReviewCommentMetadata(FReviewComment& Comment, TMap<FString, FString> Metadata)
+{
+	FString CommentText = Metadata["Body"];
+	if (Metadata.Num() == 1)
+	{
+		Comment.Body = CommentText;
+		return;
+	}
+	
+	CommentText += TEXT("\n\n[metadata]\n");
+	for (auto&[Key,Value] : Metadata)
+	{
+		if (Key != "Body")
+		{
+			CommentText += FString::Format(TEXT("{0}={1}\n"), {Key, Value});
+		}
+	}
+	Comment.Body = CommentText;
+}
+
+void FSwarmCommentsAPI::PutMetadataInBody(FReviewComment& Comment)
+{
+	TMap<FString, FString> Metadata = ParseReviewCommentMetadata(Comment);
+	if (Comment.Context.File.IsSet())
+	{
+		Metadata.FindOrAdd(TEXT("File"), Comment.Context.File.GetValue());
+		Comment.Context.File.Reset();
+	}
+	if (Comment.Context.ReplyTo.IsSet())
+	{
+		Metadata.FindOrAdd(TEXT("ReplyTo"), FString::FromInt(Comment.Context.ReplyTo.GetValue()));
+		Comment.Context.ReplyTo.Reset();
+	}
+	if (Comment.Context.Category.IsSet())
+	{
+		Metadata.FindOrAdd(TEXT("Category"), Comment.Context.Category.GetValue());
+		Comment.Context.Category.Reset();
+	}
+	if (Comment.Likes.IsSet())
+	{
+		FString LikesString;
+		TSet<FString> &Likes = Comment.Likes.GetValue();
+		if (!Likes.IsEmpty())
+		{
+			int32 Index = 0;
+			for (const FString& User : Likes)
+			{
+				LikesString += User;
+				if (Index != Likes.Num() - 1)
+				{
+					LikesString += TEXT(",");
+				}
+				++Index;
+			}
+			Metadata.FindOrAdd(TEXT("Likes"), LikesString);
+		}
+		else
+		{
+			// if Likes is set but it's empty, explicitly remove it from the metadata
+			Metadata.Remove(TEXT("Likes"));
+		}
+	}
+	SetReviewCommentMetadata(Comment, Metadata);
+}
+
+void FSwarmCommentsAPI::TakeMetadataFromBody(FReviewComment& Comment)
+{
+	// because swarm has very limited support for certain properties, they may be stored inside the body
+	// as a workaround. parse them out and put them in their proper place.
+	TMap<FString, FString> Metadata = ParseReviewCommentMetadata(Comment);
+	if (const FString* File = Metadata.Find(TEXT("File")))
+	{
+		Comment.Context.File = *File;
+	}
+	if (const FString* ReplyTo = Metadata.Find(TEXT("ReplyTo")))
+	{
+		Comment.Context.ReplyTo = FCString::Atoi(**ReplyTo);
+	}
+	if (const FString* Category = Metadata.Find(TEXT("Category")))
+	{
+		Comment.Context.Category = *Category;
+	}
+	if (const FString* Likes = Metadata.Find(TEXT("Likes")))
+	{
+		Comment.Likes = TSet<FString>();
+		TArray<FString> LikesArray;
+		Likes->ParseIntoArray(LikesArray, TEXT(","));
+		for (const FString& Like : LikesArray)
+		{
+			Comment.Likes.GetValue().Add(Like);
+		}
+	}
+	Comment.Body = Metadata[TEXT("Body")];
+}
+
+
+static FString GetEnvironmentVariable(const char* Key)
+{
+	#pragma warning(suppress : 4996) // because I'm immediately copying the result of getenv() into managed memory, it's safe.
+	if (const char* PathCStr = std::getenv(Key))
+	{
+		FString Result = FString(PathCStr);
+		FPaths::NormalizeFilename(Result);
+		return MoveTemp(Result);
+	}
+	return {};
+}
+
+// retrieves the default directory that p4tickets file is stored in
+static FString GetP4TicketsPath()
+{
+	// if the P4TICKETS environment var is set, use that path
+	const FString P4TicketsVar = GetEnvironmentVariable("P4TICKETS");
+	if (!P4TicketsVar.IsEmpty())
+	{
+		return P4TicketsVar;
+	}
+
+	// if P4TICKETS wasn't set, default to "%USERPROFILE%\p4tickets.txt" on Windows, or "$HOME/.p4tickets" otherwise
+#if PLATFORM_WINDOWS
+	return GetEnvironmentVariable("USERPROFILE") / TEXT("p4tickets.txt");
+#else
+	return GetEnvironmentVariable("HOME") / TEXT(".p4tickets");
+#endif
+}
+
+FString FSwarmCommentsAPI::RetrieveAuthorizationTicket()
+{
+	TArray<FString> TicketStrings;
+	FFileHelper::LoadFileToStringArray(TicketStrings, *GetP4TicketsPath());
+	if (TicketStrings.IsEmpty())
+	{
+		return {};
+	}
+
+	// TODO: @jordan.hoffmann: make a setting to let users choose which user they want to use
+	const FString& TicketString = TicketStrings[0];
+	
+	// find beginning of ticket
+	int32 ChopIndex = TicketString.Find(TEXT("p4d1="), ESearchCase::CaseSensitive, ESearchDir::FromStart);
+	if (ChopIndex == INDEX_NONE)
+	{
+		return {};
+	}
+	
+	// skip 'p4d1='
+	ChopIndex += 5;
+	
+	return TicketString.RightChop(ChopIndex);
+}
+
+
+FString FSwarmCommentsAPI::RetrieveSwarmURL()
+{
+	// Initialize P4Client
+	Error P4Error;
+	ClientApi P4Client;
+	P4Client.Init(&P4Error);
+	if (P4Error.Test())
+	{
+		StrBuf ErrorMsg;
+		P4Error.Fmt(&ErrorMsg);
+		UE_LOG(LogSourceControl, Error, TEXT("P4ERROR: Invalid connection to server."));
+		UE_LOG(LogSourceControl, Error, TEXT("%s"), ANSI_TO_TCHAR(ErrorMsg.Text()));
+	}
+
+	// Create a ClientUser that can capture the output from the command
+	struct MyClientUser : ClientUser
+	{
+		virtual void OutputInfo(char level, const char* data) override
+		{
+			const FString Info(data);
+			const int32 Found = Info.Find(TEXT(" = "));
+			if (Found != INDEX_NONE)
+			{
+				Result = Info.RightChop(Found + 3);
+			}
+		}
+		
+		virtual void HandleError(Error* P4Error) override
+		{
+			StrBuf ErrorMsg;
+			P4Error->Fmt(&ErrorMsg);
+			UE_LOG(LogSourceControl, Error, TEXT("P4ERROR: %s"), (ANSI_TO_TCHAR(ErrorMsg.Text())));
+		}
+
+		FString Result;
+	} P4User;
+
+	// Run property -l -n P4.Swarm.URL
+	const char* ArgV[] = { "-l", "-n", "P4.Swarm.URL" };
+	P4Client.SetArgv(3, const_cast<char* const*>(ArgV));
+	P4Client.Run("property", &P4User);
+
+	// Cleanup P4Client
+	P4Client.Final(&P4Error);
+	if (P4Error.Test())
+	{
+		StrBuf ErrorMsg;
+		P4Error.Fmt(&ErrorMsg);
+		UE_LOG(LogSourceControl, Error, TEXT("P4ERROR: Failed to disconnect from Server."));
+		UE_LOG(LogSourceControl, Error, TEXT("%s"), (ANSI_TO_TCHAR(ErrorMsg.Text())));
+	}
+
+	return P4User.Result;
+}
+
+FSwarmCommentsAPI::FAuthTicket::FAuthTicket(const FString& InUsername, const FString& InPassword)
+	: Username(InUsername), Password(InPassword)
+{}
+
+FSwarmCommentsAPI::FAuthTicket::FAuthTicket(const FString& TicketString)
+{
+	int32 ChopIndex;
+	check(TicketString.FindChar(':', ChopIndex));
+	Username = TicketString.Left(ChopIndex);
+	Password = TicketString.RightChop(ChopIndex + 1);
+}
+
+FSwarmCommentsAPI::FAuthTicket::operator FString() const
+{
+	return TEXT("Basic ") + FBase64::Encode(Username + TEXT(":") + Password);
+}
diff --git a/Engine/Plugins/Editor/ChangelistReview/Source/Private/FSwarmCommentsAPI.h b/Engine/Plugins/Editor/ChangelistReview/Source/Private/FSwarmCommentsAPI.h
new file mode 100644
--- /dev/null
+++ b/Engine/Plugins/Editor/ChangelistReview/Source/Private/FSwarmCommentsAPI.h
@@ -0,0 +1,48 @@
+// Copyright Epic Games, Inc. All Rights Reserved.
+
+#pragma once
+#include "ReviewComments.h"
+
+class FSwarmCommentsAPI : public IReviewCommentAPI
+{
+public:
+	FSwarmCommentsAPI();
+
+	virtual FString GetUsername() const override;
+
+	virtual void GetComments(const FReviewTopic& Topic, const OnGetCommentsComplete& OnComplete) const override;
+	// Body must be set
+	virtual void PostComment(FReviewComment& Comment, const OnPostCommentComplete& OnComplete) const override;
+	// all unset members will be left unchanged.
+	virtual void EditComment(const FReviewComment& Comment, const OnEditCommentComplete& OnComplete) const override;
+	
+	// retrieves the review associated with a CL. If there is none, it creates a new review and returns that.
+	virtual void GetReviewTopicForCL(const FString &ChangelistNum, const OnGetReviewTopicForCLComplete& OnComplete) const override;
+
+private:
+	void CreateReviewTopicForCL(const FString &ChangelistNum, const OnGetReviewTopicForCLComplete& OnComplete) const;
+
+	FString CommentsURL() const;
+	FString ReviewsURL() const;
+	
+	static void PutMetadataInBody(FReviewComment& Comment);
+	static void TakeMetadataFromBody(FReviewComment& Comment);
+	
+	static FString RetrieveAuthorizationTicket();
+	static FString RetrieveSwarmURL();
+
+	// used to authorize http requests to swarm
+	struct FAuthTicket
+	{
+		FAuthTicket() = default;
+		FAuthTicket(const FString& Username, const FString& Password);
+		FAuthTicket(const FString& TicketString);
+		operator FString() const;
+		FString Username;
+	private:
+		FString Password;
+	} AuthTicket;
+
+	// base url for all swarm api requests
+	FString SwarmURL;
+};
diff --git a/Engine/Source/Editor/Kismet/Private/ReviewComments.cpp b/Engine/Source/Editor/Kismet/Private/ReviewComments.cpp
new file mode 100644
--- /dev/null
+++ b/Engine/Source/Editor/Kismet/Private/ReviewComments.cpp
@@ -0,0 +1,401 @@
+// Copyright Epic Games, Inc. All Rights Reserved.
+#include "ReviewComments.h"
+#include "UObject/Class.h"
+#include "Serialization/JsonSerializer.h"
+
+namespace
+{
+
+	FString LowerCamelCase(const FString& In)
+	{
+		FString Result = In;
+		if (!Result.IsEmpty())
+		{
+			Result[0] = tolower(Result[0]);
+		}
+		return Result;
+	}
+
+	bool TryGetJsonObjectField(FReviewTopic* OutVal, const TSharedPtr<FJsonObject>& JsonObject, const FString& Key)
+	{
+		if (const TSharedPtr<FJsonValue> &Value = JsonObject->TryGetField(Key))
+		{
+			if (Value->IsNull())
+			{
+				return false;
+			}
+			*OutVal = FReviewTopic::FromString(Value->AsString());
+			return true;
+		}
+		return false;
+	}
+	
+	void TrySetJsonObjectField(TSharedPtr<FJsonObject> JsonObject, const FReviewTopic& Val, const FString& Key)
+	{
+		JsonObject->SetStringField(Key, Val.ToString());
+	}
+
+	bool TryGetJsonObjectField(int32* OutVal, const TSharedPtr<FJsonObject>& JsonObject, const FString& Key)
+	{
+		if (const TSharedPtr<FJsonValue> &Value = JsonObject->TryGetField(Key))
+		{
+			if (Value->IsNull())
+			{
+				return false;
+			}
+			*OutVal = (int32)Value->AsNumber();
+			return true;
+		}
+		return false;
+	}
+	
+	void TrySetJsonObjectField(TSharedPtr<FJsonObject> JsonObject, const int32& Val, const FString& Key)
+	{
+		JsonObject->SetNumberField(Key, Val);
+	}
+
+	bool TryGetJsonObjectField(FString* OutVal, const TSharedPtr<FJsonObject>& JsonObject, const FString& Key)
+	{
+		if (const TSharedPtr<FJsonValue> &Value = JsonObject->TryGetField(Key))
+		{
+			if (Value->IsNull())
+			{
+				return false;
+			}
+			*OutVal = Value->AsString();
+			return true;
+		}
+		return false;
+	}
+	
+	void TrySetJsonObjectField(TSharedPtr<FJsonObject> JsonObject, const FString& Val, const FString& Key)
+	{
+		JsonObject->SetStringField(Key, Val);
+	}
+
+	bool TryGetJsonObjectField(FReviewCommentContext* OutVal, const TSharedPtr<FJsonObject>& JsonObject, const FString& Key)
+	{
+		if (const TSharedPtr<FJsonValue> &Value = JsonObject->TryGetField(Key))
+		{
+			if (Value->IsNull())
+			{
+				return false;
+			}
+			// for some reason the API stores null-contexts as empty arrays. 
+			if (Value->Type == EJson::Array)
+			{
+				checkSlow(Value->AsArray().IsEmpty());
+				return false;
+			}
+			*OutVal = FReviewCommentContext::FromJson(Value->AsObject());
+			return true;
+		}
+		return false;
+	}
+	
+	void TrySetJsonObjectField(TSharedPtr<FJsonObject> JsonObject, const FReviewCommentContext& Val, const FString& Key)
+	{
+		JsonObject->SetObjectField(Key, Val.ToJsonObject());
+	}
+
+	bool TryGetJsonObjectField(TSet<FString>* OutVal, const TSharedPtr<FJsonObject>& JsonObject, const FString& Key)
+	{
+		if (const TSharedPtr<FJsonValue> &Value = JsonObject->TryGetField(Key))
+		{
+			if (Value->IsNull())
+			{
+				return false;
+			}
+			for (const TSharedPtr<FJsonValue>& Item : Value->AsArray())
+			{
+				OutVal->Add(Item->AsString());
+			}
+			return true;
+		}
+		return false;
+	}
+	
+	void TrySetJsonObjectField(TSharedPtr<FJsonObject> JsonObject, const TSet<FString>& Val, const FString& Key)
+	{
+		TArray<TSharedPtr<FJsonValue>> Array;
+		for (const FString& Element : Val)
+		{
+			Array.Add(MakeShared<FJsonValueString>(Element));
+		}
+		JsonObject->SetArrayField(Key, Array);
+	}
+
+	bool TryGetJsonObjectField(TArray<FString>* OutVal, const TSharedPtr<FJsonObject>& JsonObject, const FString& Key)
+	{
+		if (const TSharedPtr<FJsonValue> &Value = JsonObject->TryGetField(Key))
+		{
+			if (Value->IsNull())
+			{
+				return false;
+			}
+			for (const TSharedPtr<FJsonValue>& Item : Value->AsArray())
+			{
+				OutVal->Add(Item->AsString());
+			}
+			return true;
+		}
+		return false;
+	}
+	
+	void TrySetJsonObjectField(TSharedPtr<FJsonObject> JsonObject, const TArray<FString>& Val, const FString& Key)
+	{
+		TArray<TSharedPtr<FJsonValue>> Array;
+		for (const FString& Element : Val)
+		{
+			Array.Add(MakeShared<FJsonValueString>(Element));
+		}
+		JsonObject->SetArrayField(Key, Array);
+	}
+
+	bool TryGetJsonObjectField(EReviewCommentTaskState* OutVal, const TSharedPtr<FJsonObject>& JsonObject, const FString& Key)
+	{
+		if (const TSharedPtr<FJsonValue> &Value = JsonObject->TryGetField(Key))
+		{
+			if (Value->IsNull())
+			{
+				return false;
+			}
+			*OutVal = (EReviewCommentTaskState)StaticEnum<EReviewCommentTaskState>()->GetValueByNameString(Value->AsString());
+			return true;
+		}
+		return false;
+	}
+	
+	void TrySetJsonObjectField(TSharedPtr<FJsonObject> JsonObject, const EReviewCommentTaskState& Val, const FString& Key)
+	{
+		const FString EnumString = StaticEnum<EReviewCommentTaskState>()->GetValueAsString(Val);
+		int32 ChopIndex = INDEX_NONE;
+		if (EnumString.FindLastChar(':', ChopIndex))
+		{
+			JsonObject->SetStringField(Key, LowerCamelCase(EnumString.RightChop(ChopIndex + 1)));
+		}
+	}
+
+	bool TryGetJsonObjectField(FDateTime* OutVal, const TSharedPtr<FJsonObject>& JsonObject, const FString& Key)
+	{
+		if (const TSharedPtr<FJsonValue> &Value = JsonObject->TryGetField(Key))
+		{
+			if (Value->IsNull())
+			{
+				return false;
+			}
+			*OutVal = FDateTime::FromUnixTimestamp((int64)Value->AsNumber());
+			return true;
+		}
+		return false;
+	}
+	
+	void TrySetJsonObjectField(TSharedPtr<FJsonObject> JsonObject, const FDateTime& Val, const FString& Key)
+	{
+		JsonObject->SetNumberField(Key, Val.ToUnixTimestamp());
+	}
+
+	template<typename OptionalType>
+	bool TryGetJsonObjectField(TOptional<OptionalType>* OutVal, const TSharedPtr<FJsonObject>& JsonObject, const FString& Key)
+	{
+		OptionalType Val;
+		if (TryGetJsonObjectField(&Val, JsonObject, Key))
+		{
+			*OutVal = Val;
+			return true;
+		}
+		return false;
+	}
+	
+	template<typename OptionalType>
+	void TrySetJsonObjectField(TSharedPtr<FJsonObject> JsonObject, const TOptional<OptionalType>& Val, const FString& Key)
+	{
+		if (Val.IsSet())
+		{
+			TrySetJsonObjectField(JsonObject, Val.GetValue(), Key);
+		}
+	}
+
+	#define TRY_SET_FROM_JSON(OwningObject, Key, JsonObject) TryGetJsonObjectField(&(OwningObject).Key, JsonObject, LowerCamelCase(#Key))
+	#define TRY_SET_FROM_JSON_NAMED(OwningObject, Key, JsonObject, Name) TryGetJsonObjectField(&(OwningObject).Key, JsonObject, LowerCamelCase(Name))
+	#define SET_JSON_FIELD(JsonObject, OwningObject, Key) TrySetJsonObjectField(JsonObject, (OwningObject).Key, LowerCamelCase(#Key))
+	#define SET_JSON_FIELD_NAMED(JsonObject, OwningObject, Key, Name) TrySetJsonObjectField(JsonObject, (OwningObject).Key, LowerCamelCase(Name))
+}
+
+
+FString FReviewCommentContext::ToJson() const
+{
+	FString Result;
+    const TSharedRef<TJsonWriter<>> Writer = TJsonWriterFactory<>::Create(&Result);
+    FJsonSerializer::Serialize(ToJsonObject().ToSharedRef(), Writer);
+	return Result;
+}
+
+TSharedPtr<FJsonObject> FReviewCommentContext::ToJsonObject() const
+{
+	TSharedPtr<FJsonObject> Result = MakeShared<FJsonObject>();
+	SET_JSON_FIELD(Result, *this, File);
+	SET_JSON_FIELD(Result, *this, Version);
+	SET_JSON_FIELD(Result, *this, Content);
+	SET_JSON_FIELD(Result, *this, LeftLine);
+	SET_JSON_FIELD(Result, *this, RightLine);
+	SET_JSON_FIELD_NAMED(Result, *this, ReplyTo, TEXT("comment"));
+	return Result;
+}
+
+FReviewCommentContext FReviewCommentContext::FromJson(const FString& Json)
+{
+	const TSharedRef<TJsonReader<>> JsonReader = TJsonReaderFactory<>::Create(Json);
+	TSharedPtr<FJsonValue> JsonValue;
+	if (FJsonSerializer::Deserialize(JsonReader, JsonValue))
+	{
+		const TSharedPtr< FJsonObject > JsonObject = JsonValue->AsObject();
+		return FromJson(JsonObject);
+	}
+	return {};
+}
+
+FReviewCommentContext FReviewCommentContext::FromJson(const TSharedPtr<FJsonObject>& JsonObject)
+{
+	FReviewCommentContext Result = {};
+	if (!JsonObject.IsValid())
+	{
+		return Result;
+	}
+
+	TRY_SET_FROM_JSON(Result, File, JsonObject);
+	TRY_SET_FROM_JSON(Result, Version, JsonObject);
+	TRY_SET_FROM_JSON(Result, Content, JsonObject);
+	TRY_SET_FROM_JSON(Result, LeftLine, JsonObject);
+	TRY_SET_FROM_JSON(Result, RightLine, JsonObject);
+	TRY_SET_FROM_JSON_NAMED(Result, ReplyTo, JsonObject, TEXT("comment"));
+	return Result;
+}
+
+FString FReviewTopic::ToString() const
+{
+	switch(TopicType)
+	{
+	case EReviewTopicType::Review:
+		return FString::Format(TEXT("reviews/{0}"), {ChangelistNum});
+	case EReviewTopicType::Change:
+		return FString::Format(TEXT("changes/{0}"), {ChangelistNum});
+	}
+	return {};
+}
+
+FReviewTopic FReviewTopic::FromString(const FString& ReviewTopic)
+{
+	FReviewTopic Result;
+	FString TopicTypeString;
+	ReviewTopic.Split("/", &TopicTypeString, &Result.ChangelistNum);
+	if (TopicTypeString == TEXT("reviews"))
+	{
+		Result.TopicType = EReviewTopicType::Review;
+	}
+	else if (TopicTypeString == TEXT("changes"))
+	{
+		Result.TopicType = EReviewTopicType::Change;
+	}
+	else
+	{
+		Result.TopicType = EReviewTopicType::Unrecognised;
+	}
+	return Result;
+}
+
+FString FReviewComment::ToJson() const
+{
+	FString Result;
+	const TSharedRef<TJsonWriter<>> Writer = TJsonWriterFactory<>::Create(&Result);
+	FJsonSerializer::Serialize(ToJsonObject().ToSharedRef(), Writer);
+	return Result;
+}
+
+TSharedPtr<FJsonObject> FReviewComment::ToJsonObject() const
+{
+	TSharedPtr<FJsonObject> Result = MakeShared<FJsonObject>();
+	SET_JSON_FIELD(Result, *this, Topic);
+	SET_JSON_FIELD(Result, *this, Body);
+	SET_JSON_FIELD(Result, *this, Context);
+	// don't include context if it doesn't have any fields
+	const TMap<FString, TSharedPtr<FJsonValue>>& ContextValues = Result->GetObjectField(TEXT("Context"))->Values;
+	if (ContextValues.IsEmpty())
+	{
+		Result->RemoveField(TEXT("Context"));
+	}
+	SET_JSON_FIELD(Result, *this, Likes);
+	SET_JSON_FIELD(Result, *this, ReadBy);
+	SET_JSON_FIELD(Result, *this, TaskState);
+	SET_JSON_FIELD(Result, *this, User);
+	
+	SET_JSON_FIELD_NAMED(Result, *this, CommentID, TEXT("id"));
+	SET_JSON_FIELD_NAMED(Result, *this, CreatedTime, TEXT("time"));
+	SET_JSON_FIELD_NAMED(Result, *this, EditedTime, TEXT("edited"));
+	SET_JSON_FIELD_NAMED(Result, *this, UpdatedTime, TEXT("updated"));
+	if (bIsClosed)
+	{
+		TrySetJsonObjectField(Result, TArray<FString>{TEXT("closed")}, TEXT("type"));
+	}
+	
+	return Result;
+	
+}
+
+FReviewComment FReviewComment::FromJson(const FString& Json)
+{
+	const TSharedRef<TJsonReader<>> JsonReader = TJsonReaderFactory<>::Create(Json);
+	TSharedPtr<FJsonValue> JsonValue;
+	if (FJsonSerializer::Deserialize(JsonReader, JsonValue))
+	{
+		const TSharedPtr< FJsonObject > JsonObject = JsonValue->AsObject();
+		return FromJson(JsonObject);
+	}
+	return {};
+}
+
+FReviewComment FReviewComment::FromJson(const TSharedPtr<FJsonObject>& JsonObject)
+{
+	FReviewComment Result = {};
+	if (!JsonObject.IsValid())
+	{
+		return Result;
+	}
+
+	TRY_SET_FROM_JSON(Result, Topic, JsonObject);
+	TRY_SET_FROM_JSON(Result, Body, JsonObject);
+	TRY_SET_FROM_JSON(Result, Context, JsonObject);
+	TRY_SET_FROM_JSON(Result, Likes, JsonObject);
+	TRY_SET_FROM_JSON(Result, ReadBy, JsonObject);
+	TRY_SET_FROM_JSON(Result, TaskState, JsonObject);
+	TRY_SET_FROM_JSON(Result, User, JsonObject);
+	
+	TRY_SET_FROM_JSON_NAMED(Result, CommentID, JsonObject, TEXT("id"));
+	TRY_SET_FROM_JSON_NAMED(Result, CreatedTime, JsonObject, TEXT("time"));
+	TRY_SET_FROM_JSON_NAMED(Result, EditedTime, JsonObject, TEXT("edited"));
+	TRY_SET_FROM_JSON_NAMED(Result, UpdatedTime, JsonObject, TEXT("updated"));
+	
+	// set flags
+	Result.bIsClosed = false;
+	if (const TSharedPtr<FJsonValue> &FlagsField = JsonObject->TryGetField("flags"))
+	{
+		if (!FlagsField->IsNull())
+		{
+			for (const TSharedPtr<FJsonValue>& Item : FlagsField->AsArray())
+            {
+            	FString Flag = Item->AsString();
+				if (Flag == TEXT("closed"))
+				{
+					Result.bIsClosed = true;
+				}
+				else
+				{
+					// swarm documentation doesn't mention any flags except 'closed'. if this is hit, let jordan.hoffmann know
+					ensure(false);
+				}
+            }
+		}
+	}
+	
+	return Result;
+}
diff --git a/Engine/Source/Editor/Kismet/Public/ReviewComments.h b/Engine/Source/Editor/Kismet/Public/ReviewComments.h
new file mode 100644
--- /dev/null
+++ b/Engine/Source/Editor/Kismet/Public/ReviewComments.h
@@ -0,0 +1,128 @@
+// Copyright Epic Games, Inc. All Rights Reserved.
+
+#pragma once
+#include "CoreMinimal.h"
+#include "ReviewComments.generated.h"
+class FJsonObject;
+
+UENUM()
+enum class EReviewCommentTaskState : uint8
+{
+	// regular comment (not a task)
+	Comment,
+	// unaddressed task
+	Open,
+	// completed task
+	Addressed,
+	// verified completed task (documentation is limited so I'm not positive on this one)
+	Verified,
+};
+
+// swarm stores shelved file reviews separately from comments on submitted CLs
+UENUM()
+enum class EReviewTopicType : uint8
+{
+	// shelved for review
+	Review,
+	// change is already submitted
+	Change,
+
+	// this should never happen
+	Unrecognised
+};
+
+struct KISMET_API FReviewCommentContext
+{
+public:
+	// file being commented on
+	TOptional<FString> File;
+	// what version of the file this comment is associated with
+	TOptional<int32> Version;
+	// (text only) preview of text being commented on
+	TOptional<TArray<FString>> Content;
+	// (text only) line number on left panel of this comment
+	TOptional<int32> LeftLine;
+	// (text only) line number on right panel of this comment
+	TOptional<int32> RightLine;
+
+	// id of the comment that this one is replying to
+	TOptional<int32> ReplyTo;
+
+	// category this comment will appear under in the diff results
+	TOptional<FString> Category;
+	
+	FString ToJson() const;
+	TSharedPtr<FJsonObject> ToJsonObject() const;
+	static FReviewCommentContext FromJson(const FString& Json);
+	static FReviewCommentContext FromJson(const TSharedPtr<FJsonObject>& JsonObject);
+};
+
+struct KISMET_API FReviewTopic
+{
+public:
+	FString ChangelistNum;
+	EReviewTopicType TopicType;
+
+	// returns topic in format "TopicType/ChangelistNum"
+	FString ToString() const;
+	static FReviewTopic FromString(const FString& ReviewTopic);
+};
+
+struct KISMET_API FReviewComment
+{
+public:	
+	// info about CL being commented on
+	TOptional<FReviewTopic> Topic;
+	// unique id associated with this comment
+	TOptional<int32> CommentID;
+	// the comment text
+	TOptional<FString> Body;
+	
+	// names of all users who have liked this comment
+	TOptional<TSet<FString>> Likes;
+	// names of all users who have marked the comment as read
+	TOptional<TSet<FString>> ReadBy;
+	// signifies whether this is a task and if it's been addressed
+	TOptional<EReviewCommentTaskState> TaskState;
+	// comment author
+	TOptional<FString> User;
+	// when the comment was created
+	TOptional<FDateTime> CreatedTime;
+	// when the comment was last edited
+	TOptional<FDateTime> EditedTime;
+	// when the comment was created
+	TOptional<FDateTime> UpdatedTime;
+	
+	// metadata about the comment
+	FReviewCommentContext Context;
+	
+	// flags
+	bool bIsClosed : 1;
+
+	FString ToJson() const;
+	TSharedPtr<FJsonObject> ToJsonObject() const;
+	static FReviewComment FromJson(const FString& Json);
+	static FReviewComment FromJson(const TSharedPtr<FJsonObject>& JsonObject);
+};
+
+class KISMET_API IReviewCommentAPI
+{
+public:
+	DECLARE_DELEGATE_TwoParams(OnGetCommentsComplete, const TArray<FReviewComment>&, const FString& /*ErrorMessage*/)
+	DECLARE_DELEGATE_TwoParams(OnPostCommentComplete, const FReviewComment&, const FString& /*ErrorMessage*/)
+	DECLARE_DELEGATE_TwoParams(OnEditCommentComplete, const FReviewComment&, const FString& /*ErrorMessage*/)
+	DECLARE_DELEGATE_TwoParams(OnGetReviewTopicForCLComplete, const FReviewTopic&, const FString& /*ErrorMessage*/)
+	
+	virtual ~IReviewCommentAPI() = default;
+
+	// returns the name of the user logged in
+	virtual FString GetUsername() const = 0;
+	
+	virtual void GetComments(const FReviewTopic& Topic, const OnGetCommentsComplete& OnComplete) const = 0;
+	// Body must be set; username will get set by PostComment
+	virtual void PostComment(FReviewComment& Comment, const OnPostCommentComplete& OnComplete) const = 0;
+	// all unset members will be left unchanged.
+	virtual void EditComment(const FReviewComment& Comment, const OnEditCommentComplete& OnComplete) const = 0;
+	
+	virtual void GetReviewTopicForCL(const FString &ChangelistNum, const OnGetReviewTopicForCLComplete& OnComplete) const = 0;
+};
