Fix error C4700: uninitialized local variable 'ChopIndex' used when DO_CHECK is 0

diff --git a/Engine/Plugins/Editor/ChangelistReview/Source/Private/FSwarmCommentsAPI.cpp b/Engine/Plugins/Editor/ChangelistReview/Source/Private/FSwarmCommentsAPI.cpp
--- a/Engine/Plugins/Editor/ChangelistReview/Source/Private/FSwarmCommentsAPI.cpp
+++ b/Engine/Plugins/Editor/ChangelistReview/Source/Private/FSwarmCommentsAPI.cpp
@@ -521,9 +521,11 @@ FSwarmCommentsAPI::FAuthTicket::FAuthTicket(const FString& InUsername, const FSt
 FSwarmCommentsAPI::FAuthTicket::FAuthTicket(const FString& TicketString)
 {
 	int32 ChopIndex;
-	check(TicketString.FindChar(':', ChopIndex));
-	Username = TicketString.Left(ChopIndex);
-	Password = TicketString.RightChop(ChopIndex + 1);
+	if (ensure(TicketString.FindChar(':', ChopIndex)))
+	{
+		Username = TicketString.Left(ChopIndex);
+		Password = TicketString.RightChop(ChopIndex + 1);
+	}
 }
 
 FSwarmCommentsAPI::FAuthTicket::operator FString() const
