Presumptive fix for Mac builds failing due to P4 API leaking terms and causing unity build issues
#rnx

diff --git a/Engine/Plugins/Editor/ChangelistReview/Source/ChangelistReview.Build.cs b/Engine/Plugins/Editor/ChangelistReview/Source/ChangelistReview.Build.cs
--- a/Engine/Plugins/Editor/ChangelistReview/Source/ChangelistReview.Build.cs
+++ b/Engine/Plugins/Editor/ChangelistReview/Source/ChangelistReview.Build.cs
@@ -6,6 +6,13 @@ public class ChangelistReview : ModuleRules
 {
 	public ChangelistReview(ReadOnlyTargetRules Target) : base(Target)
 	{
+		// The Perforce API on Mac leaks a bunch of heavily used terms
+		// so to avoid unity builds failing, compile the module as non-unity
+		if (Target.Platform == UnrealTargetPlatform.Mac)
+		{
+			bUseUnity = false;
+		}
+
 		PublicDependencyModuleNames.AddRange(
 			new string[]
 			{
