[BugFix] Swarm Comment Crash loading changelist into the Changelist Review Tool
[BugFix] new comments don't have "reply" or "edit" buttons

diff --git a/Engine/Plugins/Editor/ChangelistReview/Source/Private/FSwarmCommentsAPI.cpp b/Engine/Plugins/Editor/ChangelistReview/Source/Private/FSwarmCommentsAPI.cpp
--- a/Engine/Plugins/Editor/ChangelistReview/Source/Private/FSwarmCommentsAPI.cpp
+++ b/Engine/Plugins/Editor/ChangelistReview/Source/Private/FSwarmCommentsAPI.cpp
@@ -321,12 +321,13 @@ void FSwarmCommentsAPI::PostComment(FReviewComment& Comment, const OnPostComment
 			}
 			if (JsonObject->HasField("comment"))
 			{
-				const FReviewComment Comment = FReviewComment::FromJson(JsonObject->GetObjectField("comment"));
+				FReviewComment Comment = FReviewComment::FromJson(JsonObject->GetObjectField("comment"));
 				if (TSharedPtr<const FSwarmCommentsAPI> Self = WeakSelf.Pin())
 				{
 					// This comment may effect the userdata. update it if necessary
 					WeakSelf.Pin()->UpdateUserdata(Comment);
 				}
+				TakeMetadataFromBody(Comment);
 				OnComplete.ExecuteIfBound(Comment, {});
 				return;
 			}
@@ -393,7 +394,9 @@ void FSwarmCommentsAPI::EditComment(const FReviewComment& Comment, const OnEditC
 			}
 			if (JsonObject->HasField("comment"))
 			{
-				OnComplete.ExecuteIfBound(FReviewComment::FromJson(JsonObject->GetObjectField("comment")), {});
+				FReviewComment Comment = FReviewComment::FromJson(JsonObject->GetObjectField("comment"));
+				TakeMetadataFromBody(Comment);
+				OnComplete.ExecuteIfBound(Comment, {});
 				return;
 			}
 		}
@@ -531,7 +534,7 @@ void FSwarmCommentsAPI::GetUserdata(const FReviewTopic& Topic, const OnGetUserDa
 	}
 
 	GetComments(Topic, OnGetCommentsComplete::CreateLambda(
-		[&OnComplete, Topic, WeakSelf = AsWeak()](const TArray<FReviewComment>& Comments, const FString& ErrorMessage)
+		[OnComplete, Topic, WeakSelf = AsWeak()](const TArray<FReviewComment>& Comments, const FString& ErrorMessage)
 		{
 			if (!ErrorMessage.IsEmpty())
 			{
