[BugFix] Review tool is hanging if it failed to connect to swarm

diff --git a/Engine/Plugins/Editor/ChangelistReview/Source/Private/SSourceControlReview.cpp b/Engine/Plugins/Editor/ChangelistReview/Source/Private/SSourceControlReview.cpp
--- a/Engine/Plugins/Editor/ChangelistReview/Source/Private/SSourceControlReview.cpp
+++ b/Engine/Plugins/Editor/ChangelistReview/Source/Private/SSourceControlReview.cpp
@@ -541,10 +541,19 @@ void SSourceControlReview::OnGetChangelistDetails(const FSourceControlOperationR
 	CommitTempChangelistNumToHistory();
 	SaveCLHistory();
 	
-	NumItemsToLoad += 2; // require review topic and comments to be loaded before continuing
 	if (CommentsAPI)
 	{
-		CommentsAPI->GetReviewTopicForCL(Changelist, FSwarmCommentsAPI::OnGetReviewTopicForCLComplete::CreateRaw(this, &SSourceControlReview::OnGetReviewTopic));
+		NumItemsToLoad += 2; // require review topic and comments to be loaded before continuing
+		if (ChangelistRecord[ReviewHelpers::AuthorKey] == TEXT("swarm"))
+		{
+			// if a swarm cl was provided, we can construct the review topic ourselves
+			OnGetReviewTopic(FReviewTopic{Changelist, EReviewTopicType::Review}, {});
+		}
+		else
+		{
+			// if a swarm cl wasn't provided, request one.
+			CommentsAPI->GetReviewTopicForCL(Changelist, FSwarmCommentsAPI::OnGetReviewTopicForCLComplete::CreateRaw(this, &SSourceControlReview::OnGetReviewTopic));
+		}
 	}
 	
 	//the loop checks if we have a valid record "depotFile(Index)" in the records to add a file entry
