[BugFix] Changelist Review Tool automatically sends a review request on Changelist load

diff --git a/Engine/Plugins/Editor/ChangelistReview/Source/Private/FSwarmCommentsAPI.cpp b/Engine/Plugins/Editor/ChangelistReview/Source/Private/FSwarmCommentsAPI.cpp
--- a/Engine/Plugins/Editor/ChangelistReview/Source/Private/FSwarmCommentsAPI.cpp
+++ b/Engine/Plugins/Editor/ChangelistReview/Source/Private/FSwarmCommentsAPI.cpp
@@ -430,7 +430,7 @@ void FSwarmCommentsAPI::GetReviewTopicForCL(const FString& ChangelistNum,
 
 	TWeakPtr<const FSwarmCommentsAPI> WeakSelf = AsWeak();
 	
-	HttpRequest->OnProcessRequestComplete().BindLambda([OnComplete, ChangelistNum, WeakSelf](FHttpRequestPtr Request, FHttpResponsePtr Response, bool bConnectedSuccessfully)
+	HttpRequest->OnProcessRequestComplete().BindLambda([OnComplete, WeakSelf](FHttpRequestPtr Request, FHttpResponsePtr Response, bool bConnectedSuccessfully)
 	{
 		const TSharedPtr<const FSwarmCommentsAPI> Self = WeakSelf.Pin();
 		if (!Self)
@@ -457,12 +457,11 @@ void FSwarmCommentsAPI::GetReviewTopicForCL(const FString& ChangelistNum,
 				const TArray<TSharedPtr<FJsonValue>> ReviewsJson = JsonObject->GetArrayField("reviews");
 				if (ReviewsJson.IsEmpty())
 				{
-					Self->CreateReviewTopicForCL(ChangelistNum, OnComplete);
+					OnComplete.ExecuteIfBound({}, TEXT("Review Not Found"));
 					return;
 				}
 				
 				const TSharedPtr<FJsonObject> Review = ReviewsJson[0]->AsObject();
-				
 				OnComplete.ExecuteIfBound(FReviewTopic{
 					FString::FromInt(Review->GetIntegerField(TEXT("id"))),
 					EReviewTopicType::Review
@@ -477,6 +476,31 @@ void FSwarmCommentsAPI::GetReviewTopicForCL(const FString& ChangelistNum,
 	HttpRequest->ProcessRequest();
 }
 
+void FSwarmCommentsAPI::GetOrCreateReviewTopicForCL(const FString& ChangelistNum,
+	const OnGetReviewTopicForCLComplete& OnComplete) const
+{
+	TWeakPtr<const FSwarmCommentsAPI> WeakSelf = AsWeak();
+	GetReviewTopicForCL(ChangelistNum,
+		OnGetReviewTopicForCLComplete::CreateLambda([OnComplete, ChangelistNum, WeakSelf](const FReviewTopic& Topic, const FString& ErrorMessage)
+		{
+			const TSharedPtr<const FSwarmCommentsAPI> Self = WeakSelf.Pin();
+			if (!Self)
+			{
+				return;
+			}
+			
+			if (ErrorMessage == TEXT("Review Not Found"))
+			{
+				Self->CreateReviewTopicForCL(ChangelistNum, OnComplete);
+			}
+			else
+			{
+				OnComplete.Execute(Topic, ErrorMessage);
+			}
+		})
+	);
+}
+
 void FSwarmCommentsAPI::CreateReviewTopicForCL(const FString& ChangelistNum, const OnGetReviewTopicForCLComplete& OnComplete) const
 {
 	
diff --git a/Engine/Plugins/Editor/ChangelistReview/Source/Private/FSwarmCommentsAPI.h b/Engine/Plugins/Editor/ChangelistReview/Source/Private/FSwarmCommentsAPI.h
--- a/Engine/Plugins/Editor/ChangelistReview/Source/Private/FSwarmCommentsAPI.h
+++ b/Engine/Plugins/Editor/ChangelistReview/Source/Private/FSwarmCommentsAPI.h
@@ -40,8 +40,11 @@ public:
 	// all unset members will be left unchanged.
 	virtual void EditComment(const FReviewComment& Comment, const OnEditCommentComplete& OnComplete, bool bSilenceNotification = false) const override;
 	
-	// retrieves the review associated with a CL. If there is none, it creates a new review and returns that.
+	// retrieves the review associated with a CL.
 	virtual void GetReviewTopicForCL(const FString &ChangelistNum, const OnGetReviewTopicForCLComplete& OnComplete) const override;
+	
+	// retrieves the review associated with a CL. If there is none, it creates a new review and returns that.
+	virtual void GetOrCreateReviewTopicForCL(const FString& ChangelistNum, const OnGetReviewTopicForCLComplete& OnComplete) const override;
 
 private:
 	
diff --git a/Engine/Plugins/Editor/ChangelistReview/Source/Private/SSourceControlReview.cpp b/Engine/Plugins/Editor/ChangelistReview/Source/Private/SSourceControlReview.cpp
--- a/Engine/Plugins/Editor/ChangelistReview/Source/Private/SSourceControlReview.cpp
+++ b/Engine/Plugins/Editor/ChangelistReview/Source/Private/SSourceControlReview.cpp
@@ -156,6 +156,21 @@ void SSourceControlReview::Construct(const FArguments& InArgs)
 							.Font(FStyleFonts::Get().Normal)
 						]
 					]
+					+SHorizontalBox::Slot()
+					.VAlign(VAlign_Center)
+					.AutoWidth()
+					[
+						SNew(SButton)
+						.HAlign(HAlign_Center)
+						.Visibility(this, &SSourceControlReview::EnableCommentsButtonVisibility)
+						.OnClicked(this, &SSourceControlReview::OnEnableCommentsButtonClicked)
+						.ToolTipText(LOCTEXT("EnableCommentsTooltip", "Create a swarm review to store comments"))
+						[
+							SNew(STextBlock)
+							.Text(LOCTEXT("EnableComments", "Enable Comments"))
+							.Font(FStyleFonts::Get().Normal)
+						]
+					]
 				]
 
 				// Author
@@ -318,6 +333,7 @@ void SSourceControlReview::LoadChangelist(const FString& Changelist)
 	GlobalReviewComments.Empty();
 	FileReviewComments.Empty();
 	TempLocalPathToDepotPath.Empty();
+	ReviewTopic.Reset();
 
 	//This command runs p4 -describe (or similar for other version controls) to retrieve changelist record information
 	GetChangelistDetailsCommand = ISourceControlOperation::Create<FGetChangelistDetails>();
@@ -337,7 +353,7 @@ const TArray<FReviewComment>* SSourceControlReview::GetReviewCommentsForFile(con
 
 void SSourceControlReview::UpdateReviewComments()
 {
-	if (bReviewCommentsDirty && CommentsAPI)
+	if (bReviewCommentsDirty && CommentsAPI && ReviewTopic)
 	{
 		CommentsAPI->GetComments(ReviewTopic.GetValue(), FSwarmCommentsAPI::OnGetCommentsComplete::CreateRaw(this, &SSourceControlReview::OnGetReviewComments));
 		bReviewCommentsDirty = false;
@@ -411,6 +427,12 @@ FString SSourceControlReview::GetReviewerUsername() const
 
 bool SSourceControlReview::IsFileInReview(const FString& File) const
 {
+	// require review comments to be valid
+	if (!CommentsAPI || !ReviewTopic.IsSet() || IsLoading())
+	{
+		return false;
+	}
+	
 	FString Path = File;
 	if (!Path.StartsWith(TEXT("//")))
 	{
@@ -635,8 +657,16 @@ void SSourceControlReview::OnGetReviewTopic(const FReviewTopic& Topic, const FSt
 {
 	if(!ErrorMessage.IsEmpty())
 	{
-		UE_LOG(LogSourceControl, Error, TEXT("IReviewCommentsAPI::GetReviewTopicForCL Error: %s"), *ErrorMessage)
-		ReviewTopic= {CurrentChangelistInfo.ChangelistNum, CurrentChangelistInfo.Author.ToString() == TEXT("Swarm") ? EReviewTopicType::Review : EReviewTopicType::Change};
+		// if a review wasn't found, we'll display a button to create a review
+		if (ErrorMessage != TEXT("Review Not Found"))
+		{
+			UE_LOG(LogSourceControl, Error, TEXT("IReviewCommentsAPI::GetReviewTopicForCL Error: %s"), *ErrorMessage);
+		}
+		
+		// skip loading comments
+		IncrementItemsLoaded();
+		IncrementItemsLoaded();
+		return;
 	}
 	else
 	{
@@ -816,6 +846,25 @@ void SSourceControlReview::OnChangelistNumCommitted(const FText& Text, ETextComm
 	}
 }
 
+FReply SSourceControlReview::OnEnableCommentsButtonClicked()
+{
+	if (CommentsAPI)
+	{
+		NumItemsToLoad += 2;
+		CommentsAPI->GetOrCreateReviewTopicForCL(CurrentChangelistInfo.ChangelistNum, FSwarmCommentsAPI::OnGetReviewTopicForCLComplete::CreateRaw(this, &SSourceControlReview::OnGetReviewTopic));
+	}
+	return FReply::Handled();
+}
+
+EVisibility SSourceControlReview::EnableCommentsButtonVisibility() const
+{
+	if (!CurrentChangelistInfo.ChangelistNum.IsEmpty() && CommentsAPI && !ReviewTopic.IsSet() && !IsLoading())
+	{
+		return EVisibility::Visible;
+	}
+	return EVisibility::Collapsed;
+}
+
 TSharedRef<SWidget> SSourceControlReview::MakeCLComboOption(TSharedPtr<FChangelistLightInfo> Item) const
 {
 	FText Text;
diff --git a/Engine/Plugins/Editor/ChangelistReview/Source/Public/SSourceControlReview.h b/Engine/Plugins/Editor/ChangelistReview/Source/Public/SSourceControlReview.h
--- a/Engine/Plugins/Editor/ChangelistReview/Source/Public/SSourceControlReview.h
+++ b/Engine/Plugins/Editor/ChangelistReview/Source/Public/SSourceControlReview.h
@@ -191,6 +191,9 @@ private:
 	FText LoadChangelistTooltip() const;
 	void OnChangelistNumChanged(const FText& Text);
 	void OnChangelistNumCommitted(const FText& Text, ETextCommit::Type CommitMethod);
+	
+	FReply OnEnableCommentsButtonClicked();
+	EVisibility EnableCommentsButtonVisibility() const;
 
 	TSharedRef<SWidget> MakeCLComboOption(TSharedPtr<SourceControlReview::FChangelistLightInfo>Item) const;
 	void OnCLComboSelection(TSharedPtr<SourceControlReview::FChangelistLightInfo> Item, ESelectInfo::Type SelectInfo) const;
diff --git a/Engine/Source/Editor/Kismet/Public/ReviewComments.h b/Engine/Source/Editor/Kismet/Public/ReviewComments.h
--- a/Engine/Source/Editor/Kismet/Public/ReviewComments.h
+++ b/Engine/Source/Editor/Kismet/Public/ReviewComments.h
@@ -132,6 +132,9 @@ public:
 	virtual void PostComment(FReviewComment& Comment, const OnPostCommentComplete& OnComplete, bool bSilenceNotification = false) const = 0;
 	// all unset members will be left unchanged.
 	virtual void EditComment(const FReviewComment& Comment, const OnEditCommentComplete& OnComplete, bool bSilenceNotification = false) const = 0;
-	
+
+	// retrieves the review associated with a CL.
 	virtual void GetReviewTopicForCL(const FString &ChangelistNum, const OnGetReviewTopicForCLComplete& OnComplete) const = 0;
+	// retrieves the review associated with a CL. If there is none, it creates a new review and returns that.
+	virtual void GetOrCreateReviewTopicForCL(const FString &ChangelistNum, const OnGetReviewTopicForCLComplete& OnComplete) const = 0;
 };
