[CDPR] Added OpenChangelist method to allow external code to open changelist review dialog

diff --git a/Engine/Plugins/Editor/ChangelistReview/Source/Private/ChangelistReviewModule.cpp b/Engine/Plugins/Editor/ChangelistReview/Source/Private/ChangelistReviewModule.cpp
--- a/Engine/Plugins/Editor/ChangelistReview/Source/Private/ChangelistReviewModule.cpp
+++ b/Engine/Plugins/Editor/ChangelistReview/Source/Private/ChangelistReviewModule.cpp
@@ -154,6 +154,33 @@ TWeakPtr<SSourceControlReview> FChangelistReviewModule::GetActiveReview()
 	return ReviewWidget;
 }
 
+bool FChangelistReviewModule::OpenChangelistReview(const FString& Changelist)
+{
+	TSharedPtr<SDockTab> DockTab = FGlobalTabmanager::Get()->TryInvokeTab(FTabId(SourceControlReviewTabName));
+	if (!DockTab)
+	{
+		return false;
+	}
+
+	TSharedPtr<SWidget> ReviewWidgetFromTab = DockTab->GetContent();
+	if (!ReviewWidgetFromTab)
+	{
+		return false;
+	}
+
+	TSharedPtr<SSourceControlReview> ReviewTool = StaticCastSharedPtr<SSourceControlReview>(ReviewWidgetFromTab);
+	if (!ReviewTool)
+	{
+		return false;
+	}
+
+	if (ReviewTool->OpenChangelist(Changelist))
+	{
+		DockTab->DrawAttention();
+		return true;
+	}
+	return false;
+}
 
 #undef LOCTEXT_NAMESPACE
 	
diff --git a/Engine/Plugins/Editor/ChangelistReview/Source/Private/SSourceControlReview.cpp b/Engine/Plugins/Editor/ChangelistReview/Source/Private/SSourceControlReview.cpp
--- a/Engine/Plugins/Editor/ChangelistReview/Source/Private/SSourceControlReview.cpp
+++ b/Engine/Plugins/Editor/ChangelistReview/Source/Private/SSourceControlReview.cpp
@@ -345,6 +345,15 @@ void SSourceControlReview::LoadChangelist(const FString& Changelist)
 		FSourceControlOperationComplete::CreateRaw(this, &SSourceControlReview::OnGetChangelistDetails, Changelist));
 }
 
+bool SSourceControlReview::OpenChangelist(const FString& Changelist)
+{
+	if (ChangelistNumText)
+	{
+		ChangelistNumText->SetText(FText::FromString(Changelist));
+	}
+	LoadChangelist(Changelist);
+	return IsLoading();
+}
 
 const TArray<FReviewComment>* SSourceControlReview::GetReviewCommentsForFile(const FString& FilePath)
 {
diff --git a/Engine/Plugins/Editor/ChangelistReview/Source/Public/ChangelistReviewModule.h b/Engine/Plugins/Editor/ChangelistReview/Source/Public/ChangelistReviewModule.h
--- a/Engine/Plugins/Editor/ChangelistReview/Source/Public/ChangelistReviewModule.h
+++ b/Engine/Plugins/Editor/ChangelistReview/Source/Public/ChangelistReviewModule.h
@@ -12,7 +12,7 @@ class SWidget;
 
 class UContentBrowserAliasDataSource;
 
-class FChangelistReviewModule : public FDefaultModuleImpl
+class CHANGELISTREVIEW_API FChangelistReviewModule : public FDefaultModuleImpl
 {
 public:
 	virtual void StartupModule() override;
@@ -21,6 +21,14 @@ public:
 	bool CanShowReviewTab() const;
 	TWeakPtr<SSourceControlReview> GetActiveReview();
 	
+	/**
+	 * Opens review tool and loads the provided change list.
+	 *
+	 * @param[in] Changelist The change list to load.
+	 * @return @c true if changelist review tool was opened and change list is put to be loaded, @c false otherwise.
+	 */
+	bool OpenChangelistReview(const FString& Changelist);
+
 	static FChangelistReviewModule& Get()
 	{
 		return FModuleManager::LoadModuleChecked<FChangelistReviewModule>("ChangelistReview");
diff --git a/Engine/Plugins/Editor/ChangelistReview/Source/Public/SSourceControlReview.h b/Engine/Plugins/Editor/ChangelistReview/Source/Public/SSourceControlReview.h
--- a/Engine/Plugins/Editor/ChangelistReview/Source/Public/SSourceControlReview.h
+++ b/Engine/Plugins/Editor/ChangelistReview/Source/Public/SSourceControlReview.h
@@ -142,6 +142,15 @@ public:
 	 */
 	void LoadChangelist(const FString& Changelist);
 
+	/**
+	 * Open changelist for review together with filling in the UI as if it was loaded for review by user.
+     *
+	 * @param Changelist CL number to review
+	 * @return @c true if changelist was open for review and is loaded, @c false otherwise.
+	 */
+	bool OpenChangelist(const FString& Changelist);
+
+
 	/**
 	 * Returns an array of review comments that are 
 	 * @param FilePath the file to retrieve comments from. (supports AssetDepotPaths and temp local file paths)
