[BugFix] Reviewing Changelist with non-asset file changes, tries to load invalid summaries.

diff --git a/Engine/Plugins/Editor/ChangelistReview/Source/Private/SSourceControlReview.cpp b/Engine/Plugins/Editor/ChangelistReview/Source/Private/SSourceControlReview.cpp
--- a/Engine/Plugins/Editor/ChangelistReview/Source/Private/SSourceControlReview.cpp
+++ b/Engine/Plugins/Editor/ChangelistReview/Source/Private/SSourceControlReview.cpp
@@ -1166,16 +1166,19 @@ void SSourceControlReview::FixupRedirectors()
 	TMap<FString, TWeakPtr<FChangelistFileData>> RedirectorsFound;
 	for (TSharedPtr<FChangelistFileData> ChangelistFile : ChangelistFiles)
 	{
-		if (FLinkerLoad* ReviewFileLoad = LoadPackageLinker(nullptr, FPackagePath::FromLocalPath(ChangelistFile->ReviewFileTempPath), LOAD_ForDiff | LOAD_DisableCompileOnLoad | LOAD_DisableEngineVersionChecks))
+		if (ChangelistFile->ReviewFileTempPath.EndsWith(TEXT(".uasset")) || ChangelistFile->ReviewFileTempPath.EndsWith(TEXT(".umap")))
 		{
-			if (IsObjectRedirector(ReviewFileLoad, ChangelistFile->AssetName))
-			{
-				FString Path = GetContentRelativePath(ReviewFileLoad);
-				if (!Path.IsEmpty())
-				{
-					RedirectorsFound.Add(Path, ChangelistFile.ToWeakPtr());
-				}
-			}
+			if (FLinkerLoad* ReviewFileLoad = LoadPackageLinker(nullptr, FPackagePath::FromLocalPath(ChangelistFile->ReviewFileTempPath), LOAD_ForDiff | LOAD_DisableCompileOnLoad | LOAD_DisableEngineVersionChecks))
+            {
+            	if (IsObjectRedirector(ReviewFileLoad, ChangelistFile->AssetName))
+            	{
+            		FString Path = GetContentRelativePath(ReviewFileLoad);
+            		if (!Path.IsEmpty())
+            		{
+            			RedirectorsFound.Add(Path, ChangelistFile.ToWeakPtr());
+            		}
+            	}
+            }
 		}
 	}
 
