[PullRquest] (modified)
PR #11184: Remove Epic hardcoded perforce server information from swarm comments tool to allow others to use it

diff --git a/Engine/Plugins/Editor/ChangelistReview/Source/Private/FSwarmCommentsAPI.cpp b/Engine/Plugins/Editor/ChangelistReview/Source/Private/FSwarmCommentsAPI.cpp
--- a/Engine/Plugins/Editor/ChangelistReview/Source/Private/FSwarmCommentsAPI.cpp
+++ b/Engine/Plugins/Editor/ChangelistReview/Source/Private/FSwarmCommentsAPI.cpp
@@ -761,22 +761,36 @@ FSwarmCommentsAPI::FAuthTicket FSwarmCommentsAPI::RetrieveAuthorizationTicket()
 		return {};
 	}
 
+	TArray<FStringView> Options;
 	for (FStringView TicketString : TicketStrings)
 	{
 		// find beginning of ticket
-		int32 ChopIndex = TicketString.Find(TEXT("p4d1="));
-		if (ChopIndex == INDEX_NONE)
+		int32 ChopIndex;
+		if (!TicketString.FindChar('=', ChopIndex))
 		{
 			continue;
 		}
-		// skip 'p4d1='
-		ChopIndex += 5;
+
+		// remove the '='
+		ChopIndex += 1;
 		
 		const FStringView Ticket = TicketString.RightChop(ChopIndex);
 		if (Ticket.StartsWith(*Username + TEXT(":"), ESearchCase::IgnoreCase))
 		{
-			return {Ticket};
+			if (TicketString.StartsWith(TEXT("localhost")))
+			{
+				return {Ticket}; // prioritize localhost if the username matches
+			}
+			Options.Add(Ticket);
+		}
+	}
+	if (!Options.IsEmpty())
+	{
+		if (Options.Num() > 1)
+		{
+			UE_LOG(LogSourceControl, Warning, TEXT("Multiple viable tickets found for p4 user. Selecting one arbitrarily"));
 		}
+		return {Options.Last()};
 	}
 	
 	return {};
