[BugFix] Blueprint Review tool always displayed the revision and the datetime of the previous asset in the diff panel as being 0. This change fixes that and properly sets the revision and the datetime

diff --git a/Engine/Plugins/Developer/PerforceSourceControl/Source/PerforceSourceControl/Private/PerforceSourceControlOperations.cpp b/Engine/Plugins/Developer/PerforceSourceControl/Source/PerforceSourceControl/Private/PerforceSourceControlOperations.cpp
--- a/Engine/Plugins/Developer/PerforceSourceControl/Source/PerforceSourceControl/Private/PerforceSourceControlOperations.cpp
+++ b/Engine/Plugins/Developer/PerforceSourceControl/Source/PerforceSourceControl/Private/PerforceSourceControlOperations.cpp
@@ -83,6 +83,7 @@ void IPerforceSourceControlWorker::RegisterWorkers()
 	WorkersMap.Add("DeleteWorkspace", FGetPerforceSourceControlWorker::CreateStatic(&InstantiateWorker<FPerforceDeleteWorkspaceWorker>));
 	WorkersMap.Add("GetFileList", FGetPerforceSourceControlWorker::CreateStatic(&InstantiateWorker<FPerforceGetFileListWorker>));
 	WorkersMap.Add("GetFile", FGetPerforceSourceControlWorker::CreateStatic(&InstantiateWorker<FPerforceGetFileWorker>));
+	WorkersMap.Add("GetSourceControlRevisionInfo", FGetPerforceSourceControlWorker::CreateStatic(&InstantiateWorker<FPerforceRevisionInfoWorker>));
 	WorkersMap.Add("Where", FGetPerforceSourceControlWorker::CreateStatic(&InstantiateWorker<FPerforceWhereWorker>));
 }
 
@@ -3414,6 +3415,65 @@ bool FPerforceDownloadFileWorker::UpdateStates() const
 	return false;
 }
 
+FName FPerforceRevisionInfoWorker::GetName() const
+{
+	return "GetSourceControlRevisionInfo";
+}
+
+bool FPerforceRevisionInfoWorker::Execute(FPerforceSourceControlCommand& InCommand)
+{
+	TRACE_CPUPROFILER_EVENT_SCOPE(FPerforceRevisionInfoWorker::Execute);
+
+	FScopedPerforceConnection ScopedConnection(InCommand);
+
+	if (!InCommand.IsCanceled() && ScopedConnection.IsValid())
+	{
+		FPerforceConnection& Connection = ScopedConnection.GetConnection();
+		
+		// TODO: The cast doesn't seem to have any type safety
+		TSharedRef<FGetSourceControlRevisionInfo, ESPMode::ThreadSafe> Operation = StaticCastSharedRef<FGetSourceControlRevisionInfo>(InCommand.Operation);
+
+		TArray<FString> Parameters;
+		FP4RecordSet Records;
+
+		// grab info about the provided files
+		for (const FString& TargetFilePath : InCommand.Files)
+		{
+			Parameters.Add(TEXT("-t"));
+			
+			Parameters.Add(TEXT("-m"));
+			Parameters.Add(TEXT("1"));
+			
+			Parameters.Add(TargetFilePath);
+
+			const ERunCommandFlags Flags = Operation->ShouldLogToStdOutput() ? ERunCommandFlags::Default : ERunCommandFlags::DisableCommandLogging;
+
+
+			InCommand.bCommandSuccessful = Connection.RunCommand(TEXT("changes"), Parameters, Records, InCommand.ResultInfo, FOnIsCancelled::CreateRaw(&InCommand, &FPerforceSourceControlCommand::IsCanceled), InCommand.bConnectionDropped);
+
+			if (InCommand.bCommandSuccessful)
+			{
+				Operation->InternalSetResults(
+					Records[0]["change"],
+					{FDateTime(1970, 1, 1, 0, 0, 0, 0) + FTimespan::FromSeconds(FCString::Atoi(*Records[0]["time"]))},
+					Records[0]["user"]);
+			}
+
+			// Keep the already allocated memory for future files to prevent reallocation
+			Parameters.Reset();
+			Records.Reset();
+		}
+	}
+	return InCommand.bCommandSuccessful;
+}
+
+bool FPerforceRevisionInfoWorker::UpdateStates() const
+{
+	// This operation will never affect the cached
+	// file states.
+	return false;
+}
+
 FName FPerforceCreateWorkspaceWorker::GetName() const
 {
 	return "CreateWorkspace";
diff --git a/Engine/Plugins/Developer/PerforceSourceControl/Source/PerforceSourceControl/Private/PerforceSourceControlOperations.h b/Engine/Plugins/Developer/PerforceSourceControl/Source/PerforceSourceControl/Private/PerforceSourceControlOperations.h
--- a/Engine/Plugins/Developer/PerforceSourceControl/Source/PerforceSourceControl/Private/PerforceSourceControlOperations.h
+++ b/Engine/Plugins/Developer/PerforceSourceControl/Source/PerforceSourceControl/Private/PerforceSourceControlOperations.h
@@ -469,6 +469,20 @@ public:
 	virtual bool UpdateStates() const override;
 };
 
+class FPerforceRevisionInfoWorker final : public IPerforceSourceControlWorker
+{
+public:
+	FPerforceRevisionInfoWorker(FPerforceSourceControlProvider& InSourceControlProvider)
+		: IPerforceSourceControlWorker(InSourceControlProvider)
+	{}
+	virtual ~FPerforceRevisionInfoWorker() = default;
+
+	// IPerforceSourceControlWorker interface
+	virtual FName GetName() const override;
+	virtual bool Execute(FPerforceSourceControlCommand& InCommand) override;
+	virtual bool UpdateStates() const override;
+};
+
 class FPerforceCreateWorkspaceWorker final : public IPerforceSourceControlWorker
 {
 public:
diff --git a/Engine/Plugins/Editor/ChangelistReview/Source/Private/SSourceControlReview.cpp b/Engine/Plugins/Editor/ChangelistReview/Source/Private/SSourceControlReview.cpp
--- a/Engine/Plugins/Editor/ChangelistReview/Source/Private/SSourceControlReview.cpp
+++ b/Engine/Plugins/Editor/ChangelistReview/Source/Private/SSourceControlReview.cpp
@@ -613,7 +613,7 @@ void SSourceControlReview::OnGetChangelistDetails(const FSourceControlOperationR
 		SetFileSourceControlAction(ChangelistFileData, ChangelistRecord[RecordActionMapKey]);
 
 		const int32 PreviousAssetRevision = bIsShelved && (ChangelistFileData->FileSourceControlAction == ESourceControlAction::Delete || ChangelistFileData->FileSourceControlAction == ESourceControlAction::Edit) ? AssetRevision : AssetRevision - 1;
-		const FString PreviousAssetRevisionStr = FString::FromInt(PreviousAssetRevision);
+		ChangelistFileData->PreviousFileRevisionNum = FString::FromInt(PreviousAssetRevision);
 		
 		TWeakPtr<SSourceControlReview> WeakReviewWidget = SharedThis(this);
 		const auto GetFileCommandResponse = [WeakReviewWidget, ChangelistFileData](const FSourceControlOperationRef& InOperation, ECommandResult::Type InResult, FString* OutFilePath)
@@ -625,24 +625,39 @@ void SSourceControlReview::OnGetChangelistDetails(const FSourceControlOperationR
 				WeakReviewWidget.Pin()->OnGetFileFromSourceControl(ChangelistFileData);
 			}
 		};
+		ISourceControlProvider& SourceControlProvider = ISourceControlModule::Get().GetProvider();
 
 		// retrieve files directly from source control into a temp location
 		if (ChangelistFileData->FileSourceControlAction != ESourceControlAction::Delete)
 		{
 			NumItemsToLoad++;
 			
-			TSharedRef<FGetFile> GetFileToReviewCommand = ISourceControlOperation::Create<FGetFile>(Changelist, ChangelistRecord[RecordRevisionMapKey], ChangelistRecord[RecordFileMapKey], bIsShelved);
+			TSharedRef<FGetFile> GetFileToReviewCommand = ISourceControlOperation::Create<FGetFile>(Changelist, ChangelistFileData->ReviewFileRevisionNum, ChangelistRecord[RecordFileMapKey], bIsShelved);
 
-			ISourceControlModule::Get().GetProvider().Execute(GetFileToReviewCommand, EConcurrency::Asynchronous, FSourceControlOperationComplete::CreateLambda(GetFileCommandResponse, &ChangelistFileData->ReviewFileTempPath));
+			SourceControlProvider.Execute(GetFileToReviewCommand, EConcurrency::Asynchronous, FSourceControlOperationComplete::CreateLambda(GetFileCommandResponse, &ChangelistFileData->ReviewFileTempPath));
 		}
 		
 		if (ChangelistFileData->FileSourceControlAction != ESourceControlAction::Add && ChangelistFileData->FileSourceControlAction != ESourceControlAction::Branch)
 		{
 			NumItemsToLoad++;
 
-			TSharedRef<FGetFile> GetPreviousFileCommand = ISourceControlOperation::Create<FGetFile>(Changelist, PreviousAssetRevisionStr, ChangelistRecord[RecordFileMapKey], false);
+			// Request the following additional info about the previous revision:
+			// - Date/Time of submission
+			// - CL number
+			FString PrevRevisionDepotPath = FString::Format(TEXT("{0}#{1}"), {ChangelistFileData->AssetDepotPath, ChangelistFileData->PreviousFileRevisionNum});
+			TSharedRef<FGetSourceControlRevisionInfo> GetSourceControlRevisionInfo = ISourceControlOperation::Create<FGetSourceControlRevisionInfo>();
+			SourceControlProvider.Execute(GetSourceControlRevisionInfo, PrevRevisionDepotPath, EConcurrency::Asynchronous, FSourceControlOperationComplete::CreateLambda(
+				[ChangelistFileData](const FSourceControlOperationRef& InOperation, ECommandResult::Type InResult)
+				{
+					const TSharedRef<FGetSourceControlRevisionInfo> Operation = StaticCastSharedRef<FGetSourceControlRevisionInfo>(InOperation);
+					ChangelistFileData->PreviousFileDateTime = Operation->GetDateTime();
+					ChangelistFileData->PreviousChangelistNum = FCString::Atoi(*Operation->GetChangelist());
+				}));
+
+			// get the contents of the previous revision
+			TSharedRef<FGetFile> GetPreviousFileCommand = ISourceControlOperation::Create<FGetFile>(Changelist, ChangelistFileData->PreviousFileRevisionNum, ChangelistRecord[RecordFileMapKey], false);
 
-			ISourceControlModule::Get().GetProvider().Execute(GetPreviousFileCommand, EConcurrency::Asynchronous, FSourceControlOperationComplete::CreateLambda(GetFileCommandResponse, &ChangelistFileData->PreviousFileTempPath));
+			SourceControlProvider.Execute(GetPreviousFileCommand, EConcurrency::Asynchronous, FSourceControlOperationComplete::CreateLambda(GetFileCommandResponse, &ChangelistFileData->PreviousFileTempPath));
 		}
 
 		//Update map keys with the next file index
diff --git a/Engine/Plugins/Editor/ChangelistReview/Source/Private/SSourceControlReviewEntry.cpp b/Engine/Plugins/Editor/ChangelistReview/Source/Private/SSourceControlReviewEntry.cpp
--- a/Engine/Plugins/Editor/ChangelistReview/Source/Private/SSourceControlReviewEntry.cpp
+++ b/Engine/Plugins/Editor/ChangelistReview/Source/Private/SSourceControlReviewEntry.cpp
@@ -353,6 +353,8 @@ FRevisionInfo SSourceControlReviewEntry::GetPreviousFileRevisionInfo() const
 {
 	//We need to have valid revision data for some DiffAssets implementations (Although for now we don't have full data on previous file we are showing correct previous revision information)
 	FRevisionInfo PreviousFileRevisionInfo;
+	PreviousFileRevisionInfo.Changelist = ChangelistFileData.PreviousChangelistNum;
+	PreviousFileRevisionInfo.Date = ChangelistFileData.PreviousFileDateTime;
 	if (ChangelistFileData.FileSourceControlAction == ESourceControlAction::Add)
 	{
 		// previous revision for added files don't exist
diff --git a/Engine/Plugins/Editor/ChangelistReview/Source/Public/SSourceControlReview.h b/Engine/Plugins/Editor/ChangelistReview/Source/Public/SSourceControlReview.h
--- a/Engine/Plugins/Editor/ChangelistReview/Source/Public/SSourceControlReview.h
+++ b/Engine/Plugins/Editor/ChangelistReview/Source/Public/SSourceControlReview.h
@@ -77,6 +77,8 @@ namespace SourceControlReview
 
 		FString PreviousFileRevisionNum;
 
+		FDateTime PreviousFileDateTime;
+
 		FString RelativeFilePath;
 
 		FString AssetFilePath;
@@ -84,6 +86,8 @@ namespace SourceControlReview
 		FString AssetDepotPath;
 
 		int32 ChangelistNum = INDEX_NONE;
+		
+		int32 PreviousChangelistNum = INDEX_NONE;
 
 		EChangelistState ChangelistState = EChangelistState::Submitted;
 
diff --git a/Engine/Source/Developer/SourceControl/Public/SourceControlOperations.h b/Engine/Source/Developer/SourceControl/Public/SourceControlOperations.h
--- a/Engine/Source/Developer/SourceControl/Public/SourceControlOperations.h
+++ b/Engine/Source/Developer/SourceControl/Public/SourceControlOperations.h
@@ -1240,6 +1240,71 @@ private:
 	TMap<FString, FSharedBuffer> FileDataMap;
 };
 
+/**
+ * Operation used to get info about file(s) at a revision. CL number, datetime, and author will be returned.
+ */
+class FGetSourceControlRevisionInfo : public FSourceControlOperationBase
+{
+public:
+
+	FGetSourceControlRevisionInfo() = default;
+
+	// ISourceControlOperation interface
+	virtual FName GetName() const override
+	{
+		return "GetSourceControlRevisionInfo";
+	}
+
+	// ISourceControlOperation interface
+	virtual FText GetInProgressString() const override
+	{
+		return LOCTEXT("SourceControl_GetRevisionInfoPending", "Getting SourceControl info for file revision...");
+	}
+
+	// ISourceControlOperation interface
+	virtual bool CanBeCalledFromBackgroundThreads() const override
+	{
+		return true;
+	}
+
+	/** Return true if the command should log its operations, otherwise false */
+	bool ShouldLogToStdOutput() const
+	{
+		return !EnumHasAllFlags(Flags, EFlags::DisableErrorLogging);
+	}
+
+	FString GetChangelist() const
+	{
+		return Changelist;
+	}
+
+	FDateTime GetDateTime() const
+	{
+		return DateTime;
+	}
+
+	FString GetAuthor() const
+	{
+		return Author;
+	}
+
+	/** 
+	 * Do not call outside of source control implementations. Used to add the results
+	 * so that the caller can get access to it.
+	 */
+	void InternalSetResults(const FString& InChangelist, const FDateTime& InDateTime, const FString& InAuthor)
+	{
+		Changelist = InChangelist;
+		DateTime = InDateTime;
+		Author = InAuthor;
+	}
+
+private:
+	FString Changelist;
+	FDateTime DateTime;
+	FString Author;
+};
+
 /** 
  * Operation used to create a new workspace if the source control system supports this functionality.
  */
