[Bugfix] casing was wrong for the silenceNotification flag in swarm api. Hopefully this should fix an issue that caused email spamming when reviewing blueprints.

diff --git a/Engine/Plugins/Editor/ChangelistReview/Source/Private/FSwarmCommentsAPI.cpp b/Engine/Plugins/Editor/ChangelistReview/Source/Private/FSwarmCommentsAPI.cpp
--- a/Engine/Plugins/Editor/ChangelistReview/Source/Private/FSwarmCommentsAPI.cpp
+++ b/Engine/Plugins/Editor/ChangelistReview/Source/Private/FSwarmCommentsAPI.cpp
@@ -24,7 +24,7 @@ namespace UE::SwarmCommentsAPI::Keywords
 	static const TCHAR* HttpPatch = TEXT("PATCH");
 
 	// swarm
-	static const TCHAR* SilenceNotification = TEXT("SilenceNotification");
+	static const TCHAR* SilenceNotification = TEXT("silenceNotification");
 	static const TCHAR* True = TEXT("true");
 	static const TCHAR* False = TEXT("false");
 
