[CrashFix] Changelist review tool was crashing when reading saved CL history. Now if the history is malformed it will be deleted

diff --git a/Engine/Plugins/Editor/ChangelistReview/Source/Private/SSourceControlReview.cpp b/Engine/Plugins/Editor/ChangelistReview/Source/Private/SSourceControlReview.cpp
--- a/Engine/Plugins/Editor/ChangelistReview/Source/Private/SSourceControlReview.cpp
+++ b/Engine/Plugins/Editor/ChangelistReview/Source/Private/SSourceControlReview.cpp
@@ -926,9 +926,10 @@ void SSourceControlReview::SaveCLHistory()
 		Authors.Add(Item->Author.ToString());
 		Descriptions.Add(Item->Description.ToString());
 	}
-	GConfig->SetArray(TEXT("SourceControlReview"), TEXT("CLHistory.Numbers"), Numbers, GEngineIni);
-	GConfig->SetArray(TEXT("SourceControlReview"), TEXT("CLHistory.Authors"), Authors, GEngineIni);
-	GConfig->SetArray(TEXT("SourceControlReview"), TEXT("CLHistory.Descriptions"), Descriptions, GEngineIni);
+	GConfig->SetArray(TEXT("SourceControlReview"), TEXT("CLHistory.Numbers"), Numbers, GEditorPerProjectIni);
+	GConfig->SetArray(TEXT("SourceControlReview"), TEXT("CLHistory.Authors"), Authors, GEditorPerProjectIni);
+	GConfig->SetArray(TEXT("SourceControlReview"), TEXT("CLHistory.Descriptions"), Descriptions, GEditorPerProjectIni);
+	GConfig->Flush(false, GEditorPerProjectIni);
 }
 
 void SSourceControlReview::LoadCLHistory()
@@ -936,18 +937,22 @@ void SSourceControlReview::LoadCLHistory()
 	TArray<FString> Numbers;
 	TArray<FString> Authors;
 	TArray<FString> Descriptions;
-	GConfig->GetArray(TEXT("SourceControlReview"), TEXT("CLHistory.Numbers"), Numbers, GEngineIni);
-	GConfig->GetArray(TEXT("SourceControlReview"), TEXT("CLHistory.Authors"), Authors, GEngineIni);
-	GConfig->GetArray(TEXT("SourceControlReview"), TEXT("CLHistory.Descriptions"), Descriptions, GEngineIni);
+	GConfig->GetArray(TEXT("SourceControlReview"), TEXT("CLHistory.Numbers"), Numbers, GEditorPerProjectIni);
+	GConfig->GetArray(TEXT("SourceControlReview"), TEXT("CLHistory.Authors"), Authors, GEditorPerProjectIni);
+	GConfig->GetArray(TEXT("SourceControlReview"), TEXT("CLHistory.Descriptions"), Descriptions, GEditorPerProjectIni);
 	CLHistory.Empty();
-	for(int32 I = 0; I < Numbers.Num(); ++I)
+	if (Numbers.Num() == Authors.Num() && Numbers.Num() == Descriptions.Num())
 	{
-		CLHistory.Add(MakeShared<FChangelistLightInfo>(
-			FText::FromString(Numbers[I]),
-			FText::FromString(Authors[I]),
-			FText::FromString(Descriptions[I])
-		));
+		for(int32 I = 0; I < Numbers.Num(); ++I)
+		{
+			CLHistory.Add(MakeShared<FChangelistLightInfo>(
+				FText::FromString(Numbers[I]),
+				FText::FromString(Authors[I]),
+				FText::FromString(Descriptions[I])
+			));
+		}
 	}
+	
 }
 
 void SSourceControlReview::SetLoading(bool bInLoading)
